Fixed zero-sized array warning. Note that this warning is explicitly disabled in MsPlatformComplierSetup.h but due to an imbalance of warning push/pop directives elsewhere, this warning has to be opted-out of again.

(peafour-cherrypick of 6835802 by [REDACTED_PERSON])

diff --git a/Engine/Source/Runtime/TraceLog/Public/Trace/Private/Event.h b/Engine/Source/Runtime/TraceLog/Public/Trace/Private/Event.h
--- a/Engine/Source/Runtime/TraceLog/Public/Trace/Private/Event.h
+++ b/Engine/Source/Runtime/TraceLog/Public/Trace/Private/Event.h
@@ -7,6 +7,11 @@
 namespace Trace
 {
 
+#if _MSC_VER
+	#pragma warning(push)
+	#pragma warning(disable : 4200) // non-standard zero-sized array
+#endif
+
 ////////////////////////////////////////////////////////////////////////////////
 struct FNewEventEvent
 {
@@ -26,6 +31,10 @@ struct FNewEventEvent
 	/*uint8		NameData[]*/
 };
 
+#if _MSC_VER
+	#pragma warning(pop)
+#endif
+
 } // namespace Trace
 
 
diff --git a/Engine/Source/Runtime/TraceLog/Public/Trace/Private/Writer.inl b/Engine/Source/Runtime/TraceLog/Public/Trace/Private/Writer.inl
--- a/Engine/Source/Runtime/TraceLog/Public/Trace/Private/Writer.inl
+++ b/Engine/Source/Runtime/TraceLog/Public/Trace/Private/Writer.inl
@@ -16,6 +16,11 @@ using TTraceAtomic = std::atomic<Type>;
 namespace Private
 {
 
+#if _MSC_VER
+	#pragma warning(push)
+	#pragma warning(disable : 4200) // non-standard zero-sized array
+#endif
+
 ////////////////////////////////////////////////////////////////////////////////
 struct FBuffer
 {
@@ -39,6 +44,10 @@ struct FBuffer
 	alignas(void*) uint8		Data[];
 };
 
+#if _MSC_VER
+	#pragma warning(pop)
+#endif
+
 extern UE_TRACE_API TTraceAtomic<FBuffer*> GActiveBuffer;
 
 static const uint16		BufferSizePow2	= 19;
