Corrected preprocessor statements to fix non-MS platforms.

(peafour-cherrypick of 6844091 by [REDACTED_PERSON])

diff --git a/Engine/Source/Runtime/TraceLog/Public/Trace/Private/Event.h b/Engine/Source/Runtime/TraceLog/Public/Trace/Private/Event.h
--- a/Engine/Source/Runtime/TraceLog/Public/Trace/Private/Event.h
+++ b/Engine/Source/Runtime/TraceLog/Public/Trace/Private/Event.h
@@ -7,7 +7,7 @@
 namespace Trace
 {
 
-#if _MSC_VER
+#if defined(_MSC_VER)
 	#pragma warning(push)
 	#pragma warning(disable : 4200) // non-standard zero-sized array
 #endif
@@ -31,7 +31,7 @@ struct FNewEventEvent
 	/*uint8		NameData[]*/
 };
 
-#if _MSC_VER
+#if defined(_MSC_VER)
 	#pragma warning(pop)
 #endif
 
diff --git a/Engine/Source/Runtime/TraceLog/Public/Trace/Private/Writer.inl b/Engine/Source/Runtime/TraceLog/Public/Trace/Private/Writer.inl
--- a/Engine/Source/Runtime/TraceLog/Public/Trace/Private/Writer.inl
+++ b/Engine/Source/Runtime/TraceLog/Public/Trace/Private/Writer.inl
@@ -16,7 +16,7 @@ using TTraceAtomic = std::atomic<Type>;
 namespace Private
 {
 
-#if _MSC_VER
+#if defined(_MSC_VER)
 	#pragma warning(push)
 	#pragma warning(disable : 4200) // non-standard zero-sized array
 #endif
@@ -44,7 +44,7 @@ struct FBuffer
 	alignas(void*) uint8		Data[];
 };
 
-#if _MSC_VER
+#if defined(_MSC_VER)
 	#pragma warning(pop)
 #endif
 
