Proper fix for "no debug symbols in executable" warnings on Mac in TraceLog module

diff --git a/Engine/Source/Runtime/TraceLog/Private/Trace/Trace.cpp b/Engine/Source/Runtime/TraceLog/Private/Trace/Trace.cpp
--- a/Engine/Source/Runtime/TraceLog/Private/Trace/Trace.cpp
+++ b/Engine/Source/Runtime/TraceLog/Private/Trace/Trace.cpp
@@ -63,8 +63,7 @@ void Flush()
 
 #else
 
-#if PLATFORM_MAC
-void Dummy() {}
-#endif // PLATFORM_MAC
+// Workaround for module not having any exported symbols
+TRACELOG_API int TraceLogExportedSymbol = 0;
 
 #endif // UE_TRACE_ENABLED
