Add UnrealInsights support for Linux (UE-79574)

diff --git a/Engine/Source/Developer/TraceInsights/Private/Insights/Common/TimeUtils.cpp b/Engine/Source/Developer/TraceInsights/Private/Insights/Common/TimeUtils.cpp
--- a/Engine/Source/Developer/TraceInsights/Private/Insights/Common/TimeUtils.cpp
+++ b/Engine/Source/Developer/TraceInsights/Private/Insights/Common/TimeUtils.cpp
@@ -194,7 +194,14 @@ FString FormatTimeMs(const double InDuration, const int32 NumDigits, bool bAddTi
 	}
 	else
 	{
+#if !PLATFORM_USE_GENERIC_STRING_IMPLEMENTATION
 		Str += FString::Printf(TEXT("%.*f"), NumDigits, Duration * 1000.0);
+#else
+		// proper resolution is tracked as UE-79534
+		TCHAR FormatString[32];
+		FCString::Snprintf(FormatString, sizeof(FormatString), TEXT("%%.%df"), NumDigits);
+		Str += FString::Printf(FormatString, Duration * 1000.0);
+#endif
 
 		if (bAddTimeUnit)
 		{
@@ -352,7 +359,14 @@ FString FormatTime(const double InTime, const double Precision)
 		{
 			Str += TEXT(" ");
 		}
+#if !PLATFORM_USE_GENERIC_STRING_IMPLEMENTATION
 		Str += FString::Printf(TEXT("%.*fs"), Digits, Time);
+#else
+		// proper resolution is tracked as UE-79534
+		TCHAR FormatString[32];
+		FCString::Snprintf(FormatString, sizeof(FormatString), TEXT("%%.%dfs"), Digits);
+		Str += FString::Printf(FormatString, Time);
+#endif
 	}
 
 	return Str;
diff --git a/Engine/Source/Programs/UnrealInsights/UnrealInsights.Target.cs b/Engine/Source/Programs/UnrealInsights/UnrealInsights.Target.cs
--- a/Engine/Source/Programs/UnrealInsights/UnrealInsights.Target.cs
+++ b/Engine/Source/Programs/UnrealInsights/UnrealInsights.Target.cs
@@ -3,7 +3,7 @@
 using UnrealBuildTool;
 using System.Collections.Generic;
 
-[SupportedPlatforms("Win64")]
+[SupportedPlatforms("Win64", "Linux")]
 public class UnrealInsightsTarget : TargetRules
 {
 	public UnrealInsightsTarget(TargetInfo Target) : base(Target)
diff --git a/Engine/Source/Runtime/Core/Public/ProfilingDebugging/CpuProfilerTrace.h b/Engine/Source/Runtime/Core/Public/ProfilingDebugging/CpuProfilerTrace.h
--- a/Engine/Source/Runtime/Core/Public/ProfilingDebugging/CpuProfilerTrace.h
+++ b/Engine/Source/Runtime/Core/Public/ProfilingDebugging/CpuProfilerTrace.h
@@ -4,7 +4,7 @@
 
 #include "CoreTypes.h"
 
-#if !IS_PROGRAM && !UE_BUILD_SHIPPING && (PLATFORM_WINDOWS || PLATFORM_PS4 || PLATFORM_XBOXONE)
+#if !IS_PROGRAM && !UE_BUILD_SHIPPING && (PLATFORM_WINDOWS || PLATFORM_PS4 || PLATFORM_XBOXONE || PLATFORM_UNIX)
 #define CPUPROFILERTRACE_ENABLED 1
 #else
 #define CPUPROFILERTRACE_ENABLED 0
diff --git a/Engine/Source/Runtime/Core/Public/Unix/UnixPlatformTime.h b/Engine/Source/Runtime/Core/Public/Unix/UnixPlatformTime.h
--- a/Engine/Source/Runtime/Core/Public/Unix/UnixPlatformTime.h
+++ b/Engine/Source/Runtime/Core/Public/Unix/UnixPlatformTime.h
@@ -26,14 +26,14 @@ struct CORE_API FUnixTime : public FGenericPlatformTime
 	static FORCEINLINE uint32 Cycles()
 	{
 		struct timespec ts;
-		clock_gettime(ClockSource, &ts);
+		clock_gettime(CLOCK_MONOTONIC, &ts);
 		return static_cast<uint32>(static_cast<uint64>(ts.tv_sec) * 1000000ULL + static_cast<uint64>(ts.tv_nsec) / 1000ULL);
 	}
 
 	static FORCEINLINE uint64 Cycles64()
 	{
 		struct timespec ts;
-		clock_gettime(CLOCK_MONOTONIC_RAW, &ts);
+		clock_gettime(CLOCK_MONOTONIC, &ts);
 		return static_cast<uint64>(static_cast<uint64>(ts.tv_sec) * 1000000ULL + static_cast<uint64>(ts.tv_nsec) / 1000ULL);
 	}
 
diff --git a/Engine/Source/Runtime/TraceLog/Private/Trace/Unix/Detail.cpp b/Engine/Source/Runtime/TraceLog/Private/Trace/Unix/Detail.cpp
new file mode 100644
--- /dev/null
+++ b/Engine/Source/Runtime/TraceLog/Private/Trace/Unix/Detail.cpp
@@ -0,0 +1,234 @@
+// Copyright 1998-2019 Epic Games, Inc. All Rights Reserved.
+
+#include "Trace/Trace.h"
+
+#if UE_TRACE_ENABLED
+
+#include <sys/mman.h>
+#include <sys/types.h>
+#include <sys/socket.h>
+#include <netdb.h>
+#include <unistd.h>
+#include <time.h>
+#include <pthread.h>
+#if defined(_GNU_SOURCE)
+	#include <sys/syscall.h>
+#endif // _GNU_SOURCE
+
+namespace Trace
+{
+
+////////////////////////////////////////////////////////////////////////////////
+uint8* MemoryReserve(SIZE_T Size)
+{
+	void* Ptr = mmap(nullptr, Size, PROT_READ | PROT_WRITE, MAP_PRIVATE | MAP_ANON, -1, 0);
+	return (Ptr != MAP_FAILED) ? reinterpret_cast<uint8*>(Ptr) : nullptr;
+}
+
+////////////////////////////////////////////////////////////////////////////////
+void MemoryFree(void* Address, SIZE_T Size)
+{
+	munmap(Address, Size);
+}
+
+////////////////////////////////////////////////////////////////////////////////
+void MemoryMap(void* Address, SIZE_T Size)
+{
+	// no-op if mmap()ed with R/W
+}
+
+////////////////////////////////////////////////////////////////////////////////
+void MemoryUnmap(void* Address, SIZE_T Size)
+{
+	madvise(Address, Size, MADV_DONTNEED);
+}
+
+
+
+////////////////////////////////////////////////////////////////////////////////
+UPTRINT ThreadCreate(const ANSICHAR* Name, void (*Entry)())
+{
+	void* (*PthreadThunk)(void*) = [] (void* Param) -> void * {
+		typedef void (*EntryType)(void);
+		(EntryType(Param))();
+		return nullptr;
+	};
+
+	pthread_t ThreadHandle;
+	if (pthread_create(&ThreadHandle, nullptr, PthreadThunk, reinterpret_cast<void *>(Entry)) != 0)
+	{
+		return 0;
+	}
+	return static_cast<UPTRINT>(ThreadHandle);
+}
+
+////////////////////////////////////////////////////////////////////////////////
+uint32 ThreadGetCurrentId()
+{
+	// FIXME: suboptimal
+	pid_t ThreadId = static_cast<pid_t>(syscall(SYS_gettid));
+	static_assert(sizeof(pid_t) <= sizeof(uint32), "pid_t is larger than uint32, reconsider implementation of GetCurrentThreadId()");
+	return static_cast<uint32>(ThreadId);
+}
+
+////////////////////////////////////////////////////////////////////////////////
+void ThreadSleep(uint32 Milliseconds)
+{
+	usleep(Milliseconds * 1000U);
+}
+
+////////////////////////////////////////////////////////////////////////////////
+void ThreadJoin(UPTRINT Handle)
+{
+	pthread_join(static_cast<pthread_t>(Handle), nullptr);
+}
+
+////////////////////////////////////////////////////////////////////////////////
+void ThreadDestroy(UPTRINT Handle)
+{
+	// no-op
+}
+
+
+
+////////////////////////////////////////////////////////////////////////////////
+uint64 TimeGetFrequency()
+{
+	return 1000000ull;
+}
+
+////////////////////////////////////////////////////////////////////////////////
+uint64 TimeGetTimestamp()
+{
+	// should stay in sync with FPlatformTime::Cycles64() or the timeline will be broken!
+	struct timespec TimeSpec;
+	clock_gettime(CLOCK_MONOTONIC, &TimeSpec);
+	return static_cast<uint64>(static_cast<uint64>(TimeSpec.tv_sec) * 1000000ULL + static_cast<uint64>(TimeSpec.tv_nsec) / 1000ULL);
+}
+
+
+
+////////////////////////////////////////////////////////////////////////////////
+UPTRINT TcpSocketConnect(const ANSICHAR* Host, uint16 Port)
+{
+	struct FAddrInfoPtr
+	{
+					~FAddrInfoPtr()	{ freeaddrinfo(Value); }
+		addrinfo*	operator -> ()	{ return Value; }
+		addrinfo**	operator & ()	{ return &Value; }
+		addrinfo*	Value;
+	};
+
+	FAddrInfoPtr Info;
+	addrinfo Hints = {};
+	Hints.ai_family = AF_INET;
+	Hints.ai_socktype = SOCK_STREAM;
+	Hints.ai_protocol = IPPROTO_TCP;
+	if (getaddrinfo(Host, nullptr, &Hints, &Info))
+	{
+		return 0;
+	}
+
+	if (&Info == nullptr)
+	{
+		return 0;
+	}
+
+	auto* SockAddr = (sockaddr_in*)Info->ai_addr;
+	SockAddr->sin_port = htons(Port);
+
+	int Socket = socket(AF_INET, SOCK_STREAM, IPPROTO_TCP);
+	if (Socket < 0)
+	{
+		return 0;
+	}
+
+	int Result = connect(Socket, Info->ai_addr, int(Info->ai_addrlen));
+	if (Result < 0)
+	{
+		close(Socket);
+		return 0;
+	}
+
+	return UPTRINT(Socket + 1);
+}
+
+////////////////////////////////////////////////////////////////////////////////
+UPTRINT TcpSocketListen(uint16 Port)
+{
+	int Socket = socket(AF_INET, SOCK_STREAM, IPPROTO_TCP);
+	if (Socket < 0)
+	{
+		return 0;
+	}
+
+	sockaddr_in SockAddr;
+	SockAddr.sin_family = AF_INET;
+	SockAddr.sin_addr.s_addr = 0;
+	SockAddr.sin_port = htons(Port);
+	int Result = bind(Socket, reinterpret_cast<sockaddr*>(&SockAddr), sizeof(SockAddr));
+	if (Result < 0)
+	{
+		close(Socket);
+		return 0;
+	}
+
+	Result = listen(Socket, 1);
+	if (Result < 0)
+	{
+		close(Socket);
+		return 0;
+	}
+
+	return UPTRINT(Socket + 1);
+}
+
+////////////////////////////////////////////////////////////////////////////////
+UPTRINT TcpSocketAccept(UPTRINT Socket)
+{
+	int Inner = Socket - 1;
+
+	Inner = accept(Inner, nullptr, nullptr);
+	if (Inner < 0)
+	{
+		return 0;
+	}
+
+	return UPTRINT(Inner + 1);
+}
+
+////////////////////////////////////////////////////////////////////////////////
+void TcpSocketClose(UPTRINT Socket)
+{
+	int Inner = Socket - 1;
+	close(Inner);
+}
+
+////////////////////////////////////////////////////////////////////////////////
+bool TcpSocketSelect(UPTRINT Socket)
+{
+	int Inner = Socket - 1;
+	fd_set FdSet;
+	FD_ZERO(&FdSet);
+	FD_SET(Inner, &FdSet);
+	timeval TimeVal = {};
+	return (select(Inner + 1, &FdSet, nullptr, nullptr, &TimeVal) != 0);
+}
+
+////////////////////////////////////////////////////////////////////////////////
+bool TcpSocketSend(UPTRINT Socket, const void* Data, uint32 Size)
+{
+	int Inner = Socket - 1;
+	return send(Inner, (const char*)Data, Size, 0) == Size;
+}
+
+////////////////////////////////////////////////////////////////////////////////
+int32 TcpSocketRecv(UPTRINT Socket, void* Data, uint32 Size)
+{
+	int Inner = Socket - 1;
+	return static_cast<int32>(recv(Inner, (char*)Data, Size, 0));
+}
+
+} // namespace Trace
+
+#endif // UE_TRACE_ENABLED
diff --git a/Engine/Source/Runtime/TraceLog/Public/Trace/Private/Trace.h b/Engine/Source/Runtime/TraceLog/Public/Trace/Private/Trace.h
--- a/Engine/Source/Runtime/TraceLog/Public/Trace/Private/Trace.h
+++ b/Engine/Source/Runtime/TraceLog/Public/Trace/Private/Trace.h
@@ -7,7 +7,7 @@
 ////////////////////////////////////////////////////////////////////////////////
 #if !defined(UE_TRACE_ENABLED)
 #	if !UE_BUILD_SHIPPING && !IS_PROGRAM
-#		if PLATFORM_WINDOWS || PLATFORM_PS4 || PLATFORM_XBOXONE
+#		if PLATFORM_WINDOWS || PLATFORM_UNIX || PLATFORM_PS4 || PLATFORM_XBOXONE
 #			define UE_TRACE_ENABLED	1
 #		endif
 #	endif
