Fix Linux Arm64 build breakage

diff --git a/Engine/Source/Runtime/TraceLog/Private/Trace/Writer.cpp b/Engine/Source/Runtime/TraceLog/Private/Trace/Writer.cpp
--- a/Engine/Source/Runtime/TraceLog/Private/Trace/Writer.cpp
+++ b/Engine/Source/Runtime/TraceLog/Private/Trace/Writer.cpp
@@ -7,7 +7,14 @@
 #include "Trace/Trace.h"
 #include "Misc/CString.h"
 
-#include <emmintrin.h>
+#if PLATFORM_CPU_X86_FAMILY
+	#include <emmintrin.h>
+	#define PLATFORM_YIELD()  _mm_pause()
+#elif PLATFORM_CPU_ARM_FAMILY
+	#define PLATFORM_YIELD()  asm volatile ("yield")
+#else
+	#error Unsupported platform!
+#endif
 
 namespace Trace
 {
@@ -39,11 +46,7 @@ namespace Private
 ////////////////////////////////////////////////////////////////////////////////
 inline void Writer_Yield()
 {
-#if PLATFORM_CPU_X86_FAMILY
-	_mm_pause();
-#else
-#	error Unsupported platform!
-#endif
+	PLATFORM_YIELD();
 }
 
 
