Fixed contention with a importer plugin of the TRACE_API define.

diff --git a/Engine/Source/Developer/TraceInsights/TraceInsights.Build.cs b/Engine/Source/Developer/TraceInsights/TraceInsights.Build.cs
--- a/Engine/Source/Developer/TraceInsights/TraceInsights.Build.cs
+++ b/Engine/Source/Developer/TraceInsights/TraceInsights.Build.cs
@@ -23,7 +23,7 @@ public class TraceInsights : ModuleRules
 				"RenderCore",
 				"Slate",
 				"EditorStyle",
-				"Trace",
+				"TraceLog",
 				"TraceServices",
 				"DesktopPlatform",
 			}
diff --git a/Engine/Source/Runtime/Core/Core.Build.cs b/Engine/Source/Runtime/Core/Core.Build.cs
--- a/Engine/Source/Runtime/Core/Core.Build.cs
+++ b/Engine/Source/Runtime/Core/Core.Build.cs
@@ -14,7 +14,7 @@ public class Core : ModuleRules
 
 		PrivateDependencyModuleNames.Add("BuildSettings");
 
-		PublicDependencyModuleNames.Add("Trace");
+		PublicDependencyModuleNames.Add("TraceLog");
 
 		PrivateIncludePaths.AddRange(
 			new string[] {
diff --git a/Engine/Source/Runtime/Launch/Launch.Build.cs b/Engine/Source/Runtime/Launch/Launch.Build.cs
--- a/Engine/Source/Runtime/Launch/Launch.Build.cs
+++ b/Engine/Source/Runtime/Launch/Launch.Build.cs
@@ -31,7 +31,7 @@ public class Launch : ModuleRules
 				"Slate",
 				"SlateCore",
 				"Sockets",
-				"Trace",
+				"TraceLog",
 				"Overlay",
 				"UtilityShaders",
 				"PreLoadScreen"
diff --git a/Engine/Source/Runtime/RHI/RHI.Build.cs b/Engine/Source/Runtime/RHI/RHI.Build.cs
--- a/Engine/Source/Runtime/RHI/RHI.Build.cs
+++ b/Engine/Source/Runtime/RHI/RHI.Build.cs
@@ -9,7 +9,7 @@ public class RHI : ModuleRules
 	{
 		PrivateDependencyModuleNames.Add("Core");
 		PrivateDependencyModuleNames.Add("ApplicationCore");
-		PrivateDependencyModuleNames.Add("Trace");
+		PrivateDependencyModuleNames.Add("TraceLog");
 
 		if (Target.bCompileAgainstEngine)
 		{
diff --git a/Engine/Source/Runtime/Trace/Private/Trace/Event.cpp b/Engine/Source/Runtime/Trace/Private/Trace/Event.cpp
deleted file mode 100644
diff --git a/Engine/Source/Runtime/Trace/Private/Trace/Trace.cpp b/Engine/Source/Runtime/Trace/Private/Trace/Trace.cpp
deleted file mode 100644
diff --git a/Engine/Source/Runtime/Trace/Private/Trace/Windows/Detail.cpp b/Engine/Source/Runtime/Trace/Private/Trace/Windows/Detail.cpp
deleted file mode 100644
diff --git a/Engine/Source/Runtime/Trace/Private/Trace/Writer.cpp b/Engine/Source/Runtime/Trace/Private/Trace/Writer.cpp
deleted file mode 100644
diff --git a/Engine/Source/Runtime/Trace/Public/Trace/Private/Event.h b/Engine/Source/Runtime/Trace/Public/Trace/Private/Event.h
deleted file mode 100644
diff --git a/Engine/Source/Runtime/Trace/Public/Trace/Private/Event.inl b/Engine/Source/Runtime/Trace/Public/Trace/Private/Event.inl
deleted file mode 100644
diff --git a/Engine/Source/Runtime/Trace/Public/Trace/Private/Field.h b/Engine/Source/Runtime/Trace/Public/Trace/Private/Field.h
deleted file mode 100644
diff --git a/Engine/Source/Runtime/Trace/Public/Trace/Private/Trace.h b/Engine/Source/Runtime/Trace/Public/Trace/Private/Trace.h
deleted file mode 100644
diff --git a/Engine/Source/Runtime/Trace/Public/Trace/Private/Writer.inl b/Engine/Source/Runtime/Trace/Public/Trace/Private/Writer.inl
deleted file mode 100644
diff --git a/Engine/Source/Runtime/Trace/Public/Trace/Trace.h b/Engine/Source/Runtime/Trace/Public/Trace/Trace.h
deleted file mode 100644
diff --git a/Engine/Source/Runtime/Trace/Trace.Build.cs b/Engine/Source/Runtime/Trace/Trace.Build.cs
deleted file mode 100644
diff --git a/Engine/Source/Runtime/TraceLog/Private/Trace/Event.cpp b/Engine/Source/Runtime/TraceLog/Private/Trace/Event.cpp
new file mode 100644
--- /dev/null
+++ b/Engine/Source/Runtime/TraceLog/Private/Trace/Event.cpp
@@ -0,0 +1,172 @@
+//
+// Copyright 1998-2019 Epic Games, Inc. All Rights Reserved.
+//
+
+#include "Trace/Private/Event.h"
+
+#if UE_TRACE_ENABLED
+
+#include "Trace/Private/Field.h"
+#include "Trace/Private/Writer.inl"
+
+namespace Trace
+{
+
+////////////////////////////////////////////////////////////////////////////////
+namespace Private
+{
+
+template <typename Type>
+struct TLateAtomic
+{
+	typedef TTraceAtomic<Type> InnerType;
+	InnerType* operator -> () { return (InnerType*)Buffer; }
+	alignas(InnerType) char Buffer[sizeof(InnerType)];
+};
+
+static TLateAtomic<uint32>	GEventUidCounter;		// = 0;
+static TLateAtomic<FEvent*>	GHeadEvent;				// = nullptr;
+void						Writer_Initialize();
+
+} // namespace Private
+
+
+
+////////////////////////////////////////////////////////////////////////////////
+static uint32 GetHashImpl(const ANSICHAR* Input, uint32 Continuation=0x811c9dc5)
+{
+	for (; *Input; ++Input)
+	{
+		Continuation ^= *Input;
+		Continuation *= 0x01000193;
+	}
+	return Continuation;
+}
+
+////////////////////////////////////////////////////////////////////////////////
+static uint32 GetLoggerHash(const ANSICHAR* LoggerName)
+{
+	return GetHashImpl(LoggerName);
+}
+
+////////////////////////////////////////////////////////////////////////////////
+static uint32 GetEventHash(const ANSICHAR* LoggerName, const ANSICHAR* EventName)
+{
+	uint32 Ret = GetHashImpl(LoggerName);
+	Ret = GetHashImpl("@", Ret);
+	Ret = GetHashImpl(EventName, Ret);
+	return Ret;
+}
+
+
+
+////////////////////////////////////////////////////////////////////////////////
+FEvent* FEvent::Find(const ANSICHAR* LoggerName, const ANSICHAR* EventName)
+{
+	uint32 EventHash = GetEventHash(LoggerName, EventName);
+
+	FEvent* Candidate = Private::GHeadEvent->load(std::memory_order_relaxed);
+	for (; Candidate != nullptr; Candidate = (FEvent*)(Candidate->Handle))
+	{
+		if (Candidate->Hash == EventHash)
+		{
+			return Candidate;
+		}
+	}
+
+	return nullptr;
+}
+
+////////////////////////////////////////////////////////////////////////////////
+void FEvent::Create(
+	FEvent* Target,
+	const FLiteralName& LoggerName,
+	const FLiteralName& EventName,
+	const FFieldDesc* FieldDescs,
+	uint32 FieldCount)
+{
+	using namespace Private;
+
+	Writer_Initialize();
+
+	// Assign a unique ID for this event
+	uint32 Uid = GEventUidCounter->fetch_add(1, std::memory_order_relaxed);
+	Uid += uint32(EKnownEventUids::User);
+
+	if (Uid >= uint32(EKnownEventUids::Max))
+	{
+		Target->bEnabled = false;
+		Target->bInitialized = true;
+		return;
+	}
+
+	// Fill out the target event's properties
+	Target->LoggerHash = GetLoggerHash(LoggerName.Ptr);
+	Target->Hash = GetEventHash(LoggerName.Ptr, EventName.Ptr);
+ 	Target->Uid = uint16(Uid);
+	Target->bEnabled = true;
+	Target->bInitialized = true;
+
+	// Calculate the number of fields and size of name data.
+	int NamesSize = LoggerName.Length + EventName.Length;
+	for (uint32 i = 0; i < FieldCount; ++i)
+	{
+		NamesSize += FieldDescs[i].NameSize;
+	}
+
+	// Allocate the new event event in the log stream.
+	uint16 EventSize = sizeof(FNewEventEvent);
+	EventSize += sizeof(FNewEventEvent::Fields[0]) * FieldCount;
+	EventSize += NamesSize;
+	auto& Event = *(FNewEventEvent*)Writer_BeginLog(uint16(EKnownEventUids::NewEvent), EventSize);
+
+	// Write event's main properties.
+	Event.Uid = Target->Uid;
+	Event.LoggerNameSize = LoggerName.Length;
+	Event.EventNameSize = EventName.Length;
+
+	// Write details about event's fields
+	Event.FieldCount = FieldCount;
+	for (uint32 i = 0; i < FieldCount; ++i)
+	{
+		const FFieldDesc& FieldDesc = FieldDescs[i];
+		auto& Out = Event.Fields[i];
+		Out.Offset = FieldDesc.ValueOffset;
+		Out.Size = FieldDesc.ValueSize;
+		Out.TypeInfo = FieldDesc.TypeInfo;
+		Out.NameSize = FieldDesc.NameSize;
+	}
+
+	// Write names
+	uint8* Cursor = (uint8*)(Event.Fields + FieldCount);
+	auto WriteName = [&Cursor] (const ANSICHAR* Data, uint32 Size)
+	{
+		memcpy(Cursor, Data, Size);
+		Cursor += Size;
+	};
+
+	WriteName(LoggerName.Ptr, LoggerName.Length);
+	WriteName(EventName.Ptr, EventName.Length);
+	for (uint32 i = 0; i < FieldCount; ++i)
+	{
+		const FFieldDesc& Desc = FieldDescs[i];
+		WriteName(Desc.Name, Desc.NameSize);
+	}
+
+	Writer_EndLog(&(uint8&)Event);
+
+	// Add this new event into the list so we can look them up later.
+	while (true)
+	{
+		FEvent* HeadEvent = GHeadEvent->load(std::memory_order_relaxed);
+		Target->Handle = HeadEvent;
+		if (GHeadEvent->compare_exchange_weak(HeadEvent, Target, std::memory_order_release))
+		{
+			break;
+		}
+	}
+}
+
+} // namespace Trace
+
+#endif // UE_TRACE_ENABLED
diff --git a/Engine/Source/Runtime/TraceLog/Private/Trace/Trace.cpp b/Engine/Source/Runtime/TraceLog/Private/Trace/Trace.cpp
new file mode 100644
--- /dev/null
+++ b/Engine/Source/Runtime/TraceLog/Private/Trace/Trace.cpp
@@ -0,0 +1,61 @@
+//
+// Copyright 1998-2019 Epic Games, Inc. All Rights Reserved.
+//
+
+#include "Trace/Trace.h"
+#include "Trace/Private/Event.h"
+
+#if UE_TRACE_ENABLED
+
+namespace Trace
+{
+
+namespace Private
+{
+
+////////////////////////////////////////////////////////////////////////////////
+extern UPTRINT	GPendingDataHandle;
+bool			Writer_ConnectImpl(const ANSICHAR*);
+bool			Writer_ToggleEventImpl(const ANSICHAR*, const ANSICHAR*, bool);
+
+} // namespace Private
+
+
+
+////////////////////////////////////////////////////////////////////////////////
+template <int DestSize>
+static void ToAnsiCheap(ANSICHAR (&Dest)[DestSize], const WIDECHAR* Src)
+{
+	for (ANSICHAR& Out : Dest)
+	{
+		Out = ANSICHAR(*Src++ & 0x7f);
+		if (Out == '\0')
+		{
+			break;
+		}
+	}
+};
+
+////////////////////////////////////////////////////////////////////////////////
+bool Connect(const TCHAR* InHost)
+{
+	char Host[32];
+	ToAnsiCheap(Host, InHost);
+
+	return Private::Writer_ConnectImpl(Host);
+}
+
+////////////////////////////////////////////////////////////////////////////////
+bool ToggleEvent(const TCHAR* LoggerName, const TCHAR* EventName, bool State)
+{
+	ANSICHAR LoggerNameA[64];
+	ANSICHAR EventNameA[64];
+	ToAnsiCheap(LoggerNameA, LoggerName);
+	ToAnsiCheap(EventNameA, EventName);
+
+	return Private::Writer_ToggleEventImpl(LoggerNameA, EventNameA, State);
+}
+
+} // namespace Trace
+
+#endif // UE_TRACE_ENABLED
diff --git a/Engine/Source/Runtime/TraceLog/Private/Trace/Windows/Detail.cpp b/Engine/Source/Runtime/TraceLog/Private/Trace/Windows/Detail.cpp
new file mode 100644
--- /dev/null
+++ b/Engine/Source/Runtime/TraceLog/Private/Trace/Windows/Detail.cpp
@@ -0,0 +1,241 @@
+//
+// Copyright 1998-2019 Epic Games, Inc. All Rights Reserved.
+//
+
+#include "Trace/Trace.h"
+
+#if UE_TRACE_ENABLED
+
+#include "Windows/AllowWindowsPlatformTypes.h"
+#	include "Windows/WindowsHWrapper.h"
+#	define _WINSOCK_DEPRECATED_NO_WARNINGS  
+#	include <winsock2.h>
+#	include <ws2tcpip.h>
+#	pragma comment(lib, "ws2_32.lib")
+#include "Windows/HideWindowsPlatformTypes.h"
+
+#pragma warning(push)
+#pragma warning(disable : 6250) // VirtualFree() missing MEM_RELEASE - a false positive
+#pragma warning(disable : 6031) // WSAStartup() return ignore  - we're error tolerant
+
+namespace Trace
+{
+
+////////////////////////////////////////////////////////////////////////////////
+uint8* MemoryReserve(SIZE_T Size)
+{
+	return (uint8*)VirtualAlloc(nullptr, Size, MEM_RESERVE, PAGE_READWRITE);
+}
+
+////////////////////////////////////////////////////////////////////////////////
+void MemoryFree(void* Address, SIZE_T Size)
+{
+	VirtualFree(Address, 0, MEM_RELEASE);
+}
+
+////////////////////////////////////////////////////////////////////////////////
+void MemoryMap(void* Address, SIZE_T Size)
+{
+	VirtualAlloc(Address, Size, MEM_COMMIT, PAGE_READWRITE);
+}
+
+////////////////////////////////////////////////////////////////////////////////
+void MemoryUnmap(void* Address, SIZE_T Size)
+{
+	VirtualFree(Address, Size, MEM_DECOMMIT);
+}
+
+
+
+////////////////////////////////////////////////////////////////////////////////
+UPTRINT ThreadCreate(const ANSICHAR* Name, void (*Entry)())
+{
+	DWORD (WINAPI *WinApiThunk)(void*) = [] (void* Param) -> DWORD {
+		typedef void (*EntryType)(void);
+		(EntryType(Param))();
+		return 0;
+	};
+
+	HANDLE Handle = CreateThread(nullptr, 0, WinApiThunk, (void*)Entry, 0, nullptr);
+	return UPTRINT(Handle);
+}
+
+////////////////////////////////////////////////////////////////////////////////
+uint32 ThreadGetCurrentId()
+{
+	return GetCurrentThreadId();
+}
+
+////////////////////////////////////////////////////////////////////////////////
+void ThreadSleep(uint32 Milliseconds)
+{
+	Sleep(Milliseconds);
+}
+
+////////////////////////////////////////////////////////////////////////////////
+void ThreadJoin(UPTRINT Handle)
+{
+	WaitForSingleObject(HANDLE(Handle), INFINITE);
+}
+
+////////////////////////////////////////////////////////////////////////////////
+void ThreadDestroy(UPTRINT Handle)
+{
+	CloseHandle(HANDLE(Handle));
+}
+
+
+
+////////////////////////////////////////////////////////////////////////////////
+uint64 TimeGetFrequency()
+{
+	LARGE_INTEGER Value;
+	QueryPerformanceFrequency(&Value);
+	return Value.QuadPart;
+}
+
+////////////////////////////////////////////////////////////////////////////////
+uint64 TimeGetTimestamp()
+{
+	LARGE_INTEGER Value;
+	QueryPerformanceCounter(&Value);
+	return Value.QuadPart;
+}
+
+
+
+////////////////////////////////////////////////////////////////////////////////
+static void TcpSocketInitialize()
+{
+	WSADATA WsaData;
+	WSAStartup(MAKEWORD(2, 2), &WsaData);
+}
+
+////////////////////////////////////////////////////////////////////////////////
+UPTRINT TcpSocketConnect(const ANSICHAR* Host, uint16 Port)
+{
+	TcpSocketInitialize();
+
+	struct FAddrInfoPtr
+	{
+					~FAddrInfoPtr()	{ freeaddrinfo(Value); }
+		addrinfo*	operator -> ()	{ return Value; }
+		addrinfo**	operator & ()	{ return &Value; }
+		addrinfo*	Value;
+	};
+
+	FAddrInfoPtr Info;
+	addrinfo Hints = {};
+	Hints.ai_family = AF_INET;
+	Hints.ai_socktype = SOCK_STREAM;
+	Hints.ai_protocol = IPPROTO_TCP;
+	if (getaddrinfo(Host, nullptr, &Hints, &Info))
+	{
+		return 0;
+	}
+
+	if (&Info == nullptr)
+	{
+		return 0;
+	}
+
+	auto* SockAddr = (sockaddr_in*)Info->ai_addr;
+	SockAddr->sin_port = htons(Port);
+
+	SOCKET Socket = socket(AF_INET, SOCK_STREAM, IPPROTO_TCP);
+	if (Socket == INVALID_SOCKET)
+	{
+		return 0;
+	}
+
+	int Result = connect(Socket, Info->ai_addr, int(Info->ai_addrlen));
+	if (Result == SOCKET_ERROR)
+	{
+		closesocket(Socket);
+		return 0;
+	}
+
+	return UPTRINT(Socket + 1);
+}
+
+////////////////////////////////////////////////////////////////////////////////
+UPTRINT TcpSocketListen(uint16 Port)
+{
+	TcpSocketInitialize();
+
+	SOCKET Socket = socket(AF_INET, SOCK_STREAM, IPPROTO_TCP);
+	if (Socket == INVALID_SOCKET)
+	{
+		return 0;
+	}
+
+	sockaddr_in SockAddr;
+	SockAddr.sin_family = AF_INET;
+	SockAddr.sin_addr.s_addr = 0;
+	SockAddr.sin_port = htons(Port);
+	int Result = bind(Socket, (SOCKADDR*)&SockAddr, sizeof(SockAddr));
+	if (Result == INVALID_SOCKET)
+	{
+		closesocket(Socket);
+		return 0;
+	}
+
+	Result = listen(Socket, 1);
+	if (Result == INVALID_SOCKET)
+	{
+		closesocket(Socket);
+		return 0;
+	}
+
+	return UPTRINT(Socket + 1);
+}
+
+////////////////////////////////////////////////////////////////////////////////
+UPTRINT TcpSocketAccept(UPTRINT Socket)
+{
+	SOCKET Inner = Socket - 1;
+
+	Inner = accept(Inner, nullptr, nullptr);
+	if (Inner == INVALID_SOCKET)
+	{
+		return 0;
+	}
+
+	return UPTRINT(Inner + 1);
+}
+
+////////////////////////////////////////////////////////////////////////////////
+void TcpSocketClose(UPTRINT Socket)
+{
+	SOCKET Inner = Socket - 1;
+	closesocket(Inner);
+}
+
+////////////////////////////////////////////////////////////////////////////////
+bool TcpSocketSelect(UPTRINT Socket)
+{
+	SOCKET Inner = Socket - 1;
+	fd_set FdSet = { 1, { Inner }, };
+	TIMEVAL TimeVal = {};
+	return (select(0, &FdSet, nullptr, nullptr, &TimeVal) != 0);
+}
+
+////////////////////////////////////////////////////////////////////////////////
+bool TcpSocketSend(UPTRINT Socket, const void* Data, uint32 Size)
+{
+	SOCKET Inner = Socket - 1;
+	return send(Inner, (const char*)Data, Size, 0) == Size;
+}
+
+////////////////////////////////////////////////////////////////////////////////
+int32 TcpSocketRecv(UPTRINT Socket, void* Data, uint32 Size)
+{
+	SOCKET Inner = Socket - 1;
+	return recv(Inner, (char*)Data, Size, 0);
+}
+
+} // namespace Trace
+
+#pragma warning(pop)
+
+#endif // UE_TRACE_ENABLED
diff --git a/Engine/Source/Runtime/TraceLog/Private/Trace/Writer.cpp b/Engine/Source/Runtime/TraceLog/Private/Trace/Writer.cpp
new file mode 100644
--- /dev/null
+++ b/Engine/Source/Runtime/TraceLog/Private/Trace/Writer.cpp
@@ -0,0 +1,865 @@
+//
+// Copyright 1998-2019 Epic Games, Inc. All Rights Reserved.
+//
+
+#include "Trace/Private/Trace.h"
+
+#if UE_TRACE_ENABLED
+
+#include "Trace/Trace.h"
+
+namespace Trace
+{
+
+////////////////////////////////////////////////////////////////////////////////
+uint8*	MemoryReserve(SIZE_T);
+void	MemoryFree(void*, SIZE_T);
+void	MemoryMap(void*, SIZE_T);
+UPTRINT	TcpSocketConnect(const ANSICHAR*, uint16);
+UPTRINT	TcpSocketListen(uint16);
+UPTRINT	TcpSocketAccept(UPTRINT);
+void	TcpSocketClose(UPTRINT);
+bool	TcpSocketSelect(UPTRINT);
+int32	TcpSocketRecv(UPTRINT, void*, uint32);
+bool	TcpSocketSend(UPTRINT, const void*, uint32);
+UPTRINT	ThreadCreate(const ANSICHAR*, void (*)());
+uint32	ThreadGetCurrentId();
+void	ThreadSleep(uint32 Milliseconds);
+void	ThreadJoin(UPTRINT);
+void	ThreadDestroy(UPTRINT);
+uint64	TimeGetFrequency();
+uint64	TimeGetTimestamp();
+
+
+
+namespace Private
+{
+
+////////////////////////////////////////////////////////////////////////////////
+inline void Writer_Yield()
+{
+#if PLATFORM_CPU_X86_FAMILY
+	_mm_pause();
+#else
+#	error Unsupported platform!
+#endif
+}
+
+
+
+////////////////////////////////////////////////////////////////////////////////
+UE_TRACE_EVENT_BEGIN($Trace, PerfWorker)
+	UE_TRACE_EVENT_FIELD(uint64, Start)
+	UE_TRACE_EVENT_FIELD(uint32, Acquire)
+	UE_TRACE_EVENT_FIELD(uint32, Send)
+	UE_TRACE_EVENT_FIELD(uint32, Done)
+UE_TRACE_EVENT_END()
+
+UE_TRACE_EVENT_BEGIN($Trace, PerfNextBuffer)
+	UE_TRACE_EVENT_FIELD(uint64, Start)
+	UE_TRACE_EVENT_FIELD(uint32, Acquire)
+	UE_TRACE_EVENT_FIELD(uint32, Done)
+	UE_TRACE_EVENT_FIELD(uint16, ThreadId)
+UE_TRACE_EVENT_END()
+
+void Writer_InitializeInstrumentation()
+{
+	// This is usually taken care of automatically by the UE_TRACE_* macros but
+	// as we're logging these events from within trace where there might not be
+	// any buffer space free, we'll explicitly initialise them here where we
+	// can guarantee we have available buffer space.
+	F$TracePerfNextBufferFields::Initialize();
+	F$TracePerfWorkerFields::Initialize();
+}
+
+
+
+////////////////////////////////////////////////////////////////////////////////
+static uint64 GStartCycle = 0;
+
+////////////////////////////////////////////////////////////////////////////////
+inline uint64 Writer_GetTimestamp()
+{
+	return TimeGetTimestamp() - GStartCycle;
+}
+
+////////////////////////////////////////////////////////////////////////////////
+void Writer_InitializeTiming()
+{
+	GStartCycle = TimeGetTimestamp();
+
+	UE_TRACE_EVENT_BEGIN($Trace, Timing)
+		UE_TRACE_EVENT_FIELD(uint64, StartCycle)
+		UE_TRACE_EVENT_FIELD(uint64, CycleFrequency)
+	UE_TRACE_EVENT_END()
+
+	UE_TRACE_LOG($Trace, Timing)
+		<< Timing.StartCycle(GStartCycle)
+		<< Timing.CycleFrequency(TimeGetFrequency());
+}
+
+
+
+////////////////////////////////////////////////////////////////////////////////
+UE_TRACE_API TTraceAtomic<FBuffer*>	GActiveBuffer;
+static FBuffer*						GTailBuffer;
+static FBuffer*						GHeadBuffer;
+static void*						GAllocBase;
+static uint32						GAllocSize;
+
+////////////////////////////////////////////////////////////////////////////////
+static void Writer_InitializeBuffers()
+{
+	const uint32 BufferCount = 4;
+
+	GAllocSize = BufferSize * (BufferCount + 1);
+	GAllocBase = MemoryReserve(GAllocSize);
+
+	UPTRINT BufferBase = UPTRINT(GAllocBase);
+	BufferBase += BufferSizeMask;
+	BufferBase &= ~UPTRINT(BufferSizeMask);
+	MemoryMap((void*)BufferBase, BufferSize * BufferCount);
+
+	FBuffer* Buffers[BufferCount];
+	for (int i = 0; i < BufferCount; ++i)
+	{
+		void* Block = (void*)(BufferBase + (BufferSize * i));
+		Buffers[i] = new (Block) FBuffer();
+	}
+
+	for (int i = 1; i < BufferCount; ++i)
+	{
+		Buffers[i - 1]->Next = Buffers[i];
+	}
+
+	GTailBuffer = Buffers[0];
+	GHeadBuffer = Buffers[BufferCount - 1];
+
+	GActiveBuffer.store(Buffers[0], std::memory_order_release);
+}
+
+////////////////////////////////////////////////////////////////////////////////
+static void Writer_ShutdownBuffers()
+{
+	MemoryFree(GAllocBase, GAllocSize);
+}
+
+////////////////////////////////////////////////////////////////////////////////
+static void Writer_RetireBuffer(FBuffer* Buffer, void (*DataSink)(const uint8*, uint32))
+{
+	uint64 StartTsc = Writer_GetTimestamp();
+
+	uint32 TailUsed;
+	for (; ; Writer_Yield())
+	{
+		TailUsed = GTailBuffer->Used.load(std::memory_order_acquire);
+		if (TailUsed < BufferRefBit)
+		{
+			break;
+		}
+	}
+	uint64 AcquireTailTsc = Writer_GetTimestamp();
+
+	if (uint32 SendSize = GTailBuffer->Final - sizeof(FBuffer))
+	{
+		DataSink(GTailBuffer->Data, SendSize);
+	}
+	uint64 SendTsc = Writer_GetTimestamp();
+
+	FBuffer* Next = GTailBuffer->Next.load(std::memory_order_relaxed);
+	new (GTailBuffer) FBuffer();
+
+	GHeadBuffer->Next.store(GTailBuffer, std::memory_order_release);
+	GHeadBuffer = GTailBuffer;
+	GTailBuffer = Next;
+	uint64 DoneTsc = Writer_GetTimestamp();
+
+	if (true)
+	{
+		uint8 PerfEventBuffer[UE_TRACE_EVENT_SIZE($Trace, PerfWorker) + FEvent::HeaderSize];
+		UE_TRACE_LOG($Trace, PerfWorker, PerfEventBuffer)
+			<< PerfWorker.Start(StartTsc)
+			<< PerfWorker.Acquire(uint32(AcquireTailTsc - StartTsc))
+			<< PerfWorker.Send(uint32(SendTsc - StartTsc))
+			<< PerfWorker.Done(uint32(DoneTsc - StartTsc));
+
+		DataSink(PerfEventBuffer, sizeof(PerfEventBuffer));
+	}
+}
+
+////////////////////////////////////////////////////////////////////////////////
+static void Writer_UpdateBuffers(void (*DataSink)(const uint8*, uint32))
+{
+	for (int i = 0; i < 3; ++i)
+	{
+		FBuffer* Buffer = GActiveBuffer.load(std::memory_order_acquire);
+		if (GTailBuffer == Buffer)
+		{
+			return;
+		}
+
+		Writer_RetireBuffer(Buffer, DataSink);
+	}
+}
+
+////////////////////////////////////////////////////////////////////////////////
+static void Writer_Flush()
+{
+	using namespace Private;
+
+	FBuffer* Buffer = GActiveBuffer.load(std::memory_order_acquire);
+	uint32 PrevUsed = Buffer->Used.fetch_add(BufferSize, std::memory_order_relaxed);
+
+	if (!(PrevUsed & BufferSize))
+	{
+		FBuffer* NextBuffer;
+		for (; ; Writer_Yield())
+		{
+			NextBuffer = Buffer->Next.load(std::memory_order_relaxed);
+			if (NextBuffer != nullptr)
+			{
+				break;
+			}
+		}
+
+		Buffer->Final = PrevUsed & BufferSizeMask;
+		GActiveBuffer.store(NextBuffer, std::memory_order_release);
+	}
+}
+
+////////////////////////////////////////////////////////////////////////////////
+UE_TRACE_API void* Writer_NextBuffer(Private::FBuffer* Buffer, uint32 PrevUsed, uint32 SizeAndRef)
+{
+	using namespace Private;
+
+	uint32 PerfEventSize = UE_TRACE_EVENT_SIZE($Trace, PerfNextBuffer) + FEvent::HeaderSize;
+	SizeAndRef += PerfEventSize;
+
+	uint64 StartTsc = Writer_GetTimestamp();
+
+	FBuffer* NextBuffer;
+	while (true)
+	{
+		if (!(PrevUsed & BufferSize))
+		{
+			Buffer->Final = PrevUsed & BufferSizeMask;
+		}
+
+		// Get the next candidate buffer to allocate from.
+		for (; ; Writer_Yield())
+		{
+			NextBuffer = Buffer->Next.load(std::memory_order_relaxed);
+			if (NextBuffer != nullptr)
+			{
+				break;
+			}
+		}
+
+		Buffer->Used.fetch_sub(BufferRefBit, std::memory_order_release);
+
+		// Try and allocate some space in the next buffer.
+		PrevUsed = NextBuffer->Used.fetch_add(SizeAndRef, std::memory_order_relaxed);
+		uint32 Used = PrevUsed + SizeAndRef;
+		if (UNLIKELY(!(Used & BufferSize)))
+		{
+			break;
+		}
+
+		// Next buffer's full. Try again.
+		Buffer = NextBuffer;
+	}
+
+	uint64 AcquireTsc = Writer_GetTimestamp();
+
+	if (!(PrevUsed & BufferSize))
+	{
+		GActiveBuffer.compare_exchange_weak(Buffer, NextBuffer, std::memory_order_release);
+	}
+
+	PrevUsed &= BufferSizeMask;
+	uint8* Out = (uint8*)(UPTRINT(NextBuffer) + PrevUsed);
+
+	uint64 DoneTsc = Writer_GetTimestamp();
+	UE_TRACE_LOG($Trace, PerfNextBuffer, Out)
+		<< PerfNextBuffer.Start(StartTsc)
+		<< PerfNextBuffer.Acquire(uint32(AcquireTsc - StartTsc))
+		<< PerfNextBuffer.Done(uint32(DoneTsc - StartTsc))
+		<< PerfNextBuffer.ThreadId(uint16(ThreadGetCurrentId()));
+
+	return Out + PerfEventSize;
+}
+
+
+
+////////////////////////////////////////////////////////////////////////////////
+class FHoldBuffer
+{
+public:
+	void				Init();
+	void				Shutdown();
+	void				Write(const void* Data, uint32 Size);
+	bool				IsFull() const	{ return Full; }
+	const uint8*		GetData() const { return Base; }
+	uint32				GetSize() const { return Used; }
+
+private:
+	static const uint32	PageShift = 16;
+	static const uint32	PageSize = 1 << PageShift;
+	static const uint32	MaxPages = (4 * 1024 * 1024) >> PageShift;
+	uint8*				Base;
+	int32				Used;
+	uint16				MappedPageCount;
+	bool				Full;
+};
+
+////////////////////////////////////////////////////////////////////////////////
+void FHoldBuffer::Init()
+{
+	Base = MemoryReserve(FHoldBuffer::PageSize * FHoldBuffer::MaxPages);
+	Used = 0;
+	MappedPageCount = 0;
+	Full = false;
+}
+
+////////////////////////////////////////////////////////////////////////////////
+void FHoldBuffer::Shutdown()
+{
+	if (Base == nullptr)
+	{
+		return;
+	}
+
+	MemoryFree(Base, FHoldBuffer::PageSize * FHoldBuffer::MaxPages);
+	Base = nullptr;
+	MappedPageCount = 0;
+	Used = 0;
+}
+
+////////////////////////////////////////////////////////////////////////////////
+void FHoldBuffer::Write(const void* Data, uint32 Size)
+{
+	int32 NextUsed = Used + Size;
+
+	uint16 HotPageCount = uint16((NextUsed + (FHoldBuffer::PageSize - 1)) >> FHoldBuffer::PageShift);
+	if (HotPageCount > MappedPageCount)
+	{
+		if (HotPageCount >= FHoldBuffer::MaxPages)
+		{
+			Full = true;
+			return;
+		}
+
+		void* MapStart = Base + (UPTRINT(MappedPageCount) << FHoldBuffer::PageShift);
+		uint32 MapSize = (HotPageCount - MappedPageCount) << FHoldBuffer::PageShift;
+		MemoryMap(MapStart, MapSize);
+
+		MappedPageCount = HotPageCount;
+	}
+
+	memcpy(Base + Used, Data, Size);
+
+	Used = NextUsed;
+}
+
+
+
+////////////////////////////////////////////////////////////////////////////////
+enum class EDataState : uint8
+{
+	Passive = 0,		// Data is being collected in-process
+	Partial,			// Passive, but buffers are full so some events are lost
+	Sending,			// Events are being sent to an IO handle
+};
+static FHoldBuffer		GHoldBuffer;
+static UPTRINT			GDataHandle			= 0;
+static EDataState		GDataState;			// = EDataState::Passive;
+UPTRINT					GPendingDataHandle	= 0;
+
+////////////////////////////////////////////////////////////////////////////////
+static void Writer_UpdateData()
+{
+	if (GPendingDataHandle)
+	{
+		if (GDataHandle)
+		{
+			TcpSocketClose(GPendingDataHandle);
+			GPendingDataHandle = 0;
+		}
+		else
+		{
+			GDataHandle = GPendingDataHandle;
+			GPendingDataHandle = 0;
+			GDataState = EDataState::Sending;
+
+			// Handshake.
+			const uint32 Magic = 'TRCE';
+			TcpSocketSend(GDataHandle, &Magic, sizeof(Magic));
+
+			// Stream header
+			const struct {
+				uint8 Format;
+				uint8 Parameter;
+			} TransportHeader = { 1 };
+			TcpSocketSend(GDataHandle, &TransportHeader, sizeof(TransportHeader));
+
+			// Passively collected data
+			TcpSocketSend(GDataHandle, GHoldBuffer.GetData(), GHoldBuffer.GetSize());
+			GHoldBuffer.Shutdown();
+		}
+	}
+
+	// Passive mode?
+	if (GDataState == EDataState::Sending)
+	{
+		// Transmit data to the io handle
+		if (GDataHandle)
+		{
+			Writer_UpdateBuffers([] (const uint8* Data, uint32 Size)
+			{
+				if (GDataHandle && !TcpSocketSend(GDataHandle, Data, Size))
+				{
+					TcpSocketClose(GDataHandle);
+					GDataHandle = 0;
+				}
+			});
+		}
+		else
+		{
+			Writer_UpdateBuffers([] (const uint8*, uint32) {});
+		}
+	}
+	else
+	{
+		// Send data to hold/ring
+		Writer_UpdateBuffers([] (const uint8* Data, uint32 Size)
+		{
+			GHoldBuffer.Write(Data, Size);
+		});
+
+		// Did we overflow? Enter partial mode.
+		bool Overflown = GHoldBuffer.IsFull();
+		if (Overflown && GDataState != EDataState::Partial)
+		{
+			/* TODO: turn off non-retail events */
+			GDataState = EDataState::Partial;
+		}
+	}
+}
+
+
+
+////////////////////////////////////////////////////////////////////////////////
+enum class EControlState : uint8
+{
+	Closed = 0,
+	Listening,
+	Accepted,
+	Failed,
+};
+
+struct FControlCommands
+{
+	enum { Max = 3 };
+	struct
+	{
+		uint32	Hash;
+		void*	Param;
+		void	(*Thunk)(void*, uint32, ANSICHAR const* const*);
+	}			Commands[Max];
+	uint8		Count;
+};
+
+////////////////////////////////////////////////////////////////////////////////
+bool Writer_ConnectImpl(const ANSICHAR*);
+bool Writer_ToggleEventImpl(const ANSICHAR*, const ANSICHAR*, bool);
+
+////////////////////////////////////////////////////////////////////////////////
+static FControlCommands	GControlCommands;
+static UPTRINT			GControlListen		= 0;
+static UPTRINT			GControlSocket		= 0;
+static EControlState	GControlState;		// = EControlState::Closed;
+
+////////////////////////////////////////////////////////////////////////////////
+static uint32 Writer_ControlHash(const ANSICHAR* Word)
+{
+	uint32 Hash = 5381;
+	for (; *Word; (Hash = (Hash * 33) ^ *Word), ++Word);
+	return Hash;
+}
+
+////////////////////////////////////////////////////////////////////////////////
+static bool Writer_ControlAddCommand(
+	const ANSICHAR* Name,
+	void* Param,
+	void (*Thunk)(void*, uint32, ANSICHAR const* const*))
+{
+	if (GControlCommands.Count >= FControlCommands::Max)
+	{
+		return false;
+	}
+
+	uint32 Index = GControlCommands.Count++;
+	GControlCommands.Commands[Index] = { Writer_ControlHash(Name), Param, Thunk };
+	return true;
+}
+
+////////////////////////////////////////////////////////////////////////////////
+static bool Writer_ControlDispatch(uint32 ArgC, ANSICHAR const* const* ArgV)
+{
+	if (ArgC == 0)
+	{
+		return false;
+	}
+
+	uint32 Hash = Writer_ControlHash(ArgV[0]);
+	--ArgC;
+	++ArgV;
+
+	for (int i = 0, n = GControlCommands.Count; i < n; ++i)
+	{
+		const auto& Command = GControlCommands.Commands[i];
+		if (Command.Hash == Hash)
+		{
+			Command.Thunk(Command.Param, ArgC, ArgV);
+			return true;
+		}
+	}
+
+	return false;
+}
+
+////////////////////////////////////////////////////////////////////////////////
+static bool Writer_ControlListen()
+{
+#if PLATFORM_PS4
+	GControlState = EControlState::Failed;
+	return false;
+#endif
+
+	GControlListen = TcpSocketListen(1985);
+	if (!GControlListen)
+	{
+		GControlState = EControlState::Failed;
+		return false;
+	}
+
+	GControlState = EControlState::Listening;
+	return true;
+}
+
+////////////////////////////////////////////////////////////////////////////////
+static bool Writer_ControlAccept()
+{
+	if (!TcpSocketSelect(GControlListen))
+	{
+		return false;
+	}
+
+	UPTRINT Socket = TcpSocketAccept(GControlListen);
+	if (!Socket)
+	{
+		return false;
+	}
+
+	GControlState = EControlState::Accepted;
+	GControlSocket = Socket;
+	return true;
+}
+
+////////////////////////////////////////////////////////////////////////////////
+static void Writer_ControlRecv()
+{
+	// We'll assume that commands are smaller than the canonical MTU so this
+	// doesn't need to be implemented in a reentrant manner (maybe).
+
+	ANSICHAR Buffer[512];
+	ANSICHAR* __restrict Head = Buffer;
+	while (TcpSocketSelect(GControlSocket))
+	{
+		int32 ReadSize = int32(UPTRINT(Buffer + sizeof(Buffer) - Head));
+		int32 Recvd = TcpSocketRecv(GControlSocket, Head, ReadSize);
+		if (Recvd <= 0)
+		{
+			TcpSocketClose(GControlSocket);
+			GControlSocket = 0;
+			GControlState = EControlState::Listening;
+			break;
+		}
+
+		Head += Recvd;
+
+		enum EParseState
+		{
+			CrLfSkip,
+			WhitespaceSkip,
+			Word,
+		} ParseState = EParseState::CrLfSkip;
+
+		const uint32 ArgN = 16;
+		uint32 ArgC = 0;
+		const ANSICHAR* ArgV[ArgN];
+
+		const ANSICHAR* __restrict Spent = Buffer;
+		for (ANSICHAR* __restrict Cursor = Buffer; Cursor < Head; ++Cursor)
+		{
+			switch (ParseState)
+			{
+			case EParseState::CrLfSkip:
+				if (*Cursor == '\n' || *Cursor == '\r')
+				{
+					continue;
+				}
+				ParseState = EParseState::WhitespaceSkip;
+				/* [[fallthrough]] */
+
+			case EParseState::WhitespaceSkip:
+				if (*Cursor == ' ' || *Cursor == '\0')
+				{
+					continue;
+				}
+
+				if (ArgC < ArgN)
+				{
+					ArgV[ArgC] = Cursor;
+					++ArgC;
+				}
+
+				ParseState = EParseState::Word;
+				/* [[fallthrough]] */
+
+			case EParseState::Word:
+				if (*Cursor == ' ' || *Cursor == '\0')
+				{
+					*Cursor = '\0';
+					ParseState = EParseState::WhitespaceSkip;
+					continue;
+				}
+
+				if (*Cursor == '\r' || *Cursor == '\n')
+				{
+					*Cursor = '\0';
+
+					Writer_ControlDispatch(ArgC, ArgV);
+
+					ArgC = 0;
+					Spent = Cursor + 1;
+					ParseState = EParseState::CrLfSkip;
+					continue;
+				}
+
+				break;
+			}
+		}
+
+		int32 UnspentSize = int32(UPTRINT(Head - Spent));
+		if (UnspentSize)
+		{
+			memmove(Buffer, Spent, UnspentSize);
+		}
+		Head = Buffer + UnspentSize;
+	}
+}
+
+////////////////////////////////////////////////////////////////////////////////
+static void Writer_UpdateControl()
+{
+	switch (GControlState)
+	{
+	case EControlState::Closed:
+		if (!Writer_ControlListen())
+		{
+			break;
+		}
+		/* [[fallthrough]] */
+
+	case EControlState::Listening:
+		if (!Writer_ControlAccept())
+		{
+			break;
+		}
+		/* [[fallthrough]] */
+
+	case EControlState::Accepted:
+		Writer_ControlRecv();
+		break;
+	}
+}
+
+////////////////////////////////////////////////////////////////////////////////
+static void Writer_InitializeControl()
+{
+	Writer_ControlListen();
+
+	Writer_ControlAddCommand("Connect", nullptr,
+		[] (void*, uint32 ArgC, ANSICHAR const* const* ArgV)
+		{
+			if (ArgC > 0)
+			{
+				Writer_ConnectImpl(ArgV[0]);
+			}
+		}
+	);
+
+	Writer_ControlAddCommand("ToggleEvent", nullptr,
+		[] (void*, uint32 ArgC, ANSICHAR const* const* ArgV)
+		{
+			if (ArgC < 2)
+			{
+				return;
+			}
+			const ANSICHAR* LoggerName = ArgV[0];
+			const ANSICHAR* EventName = ArgV[1];
+			const ANSICHAR* State = (ArgC > 2) ? ArgV[2] : "";
+			Writer_ToggleEventImpl(LoggerName, EventName, State[0] != '0');
+		}
+	);
+}
+
+////////////////////////////////////////////////////////////////////////////////
+static void Writer_ShutdownControl()
+{
+	if (GControlListen)
+	{
+		TcpSocketClose(GControlListen);
+		GControlListen = 0;
+	}
+}
+
+
+
+////////////////////////////////////////////////////////////////////////////////
+static UPTRINT			GWorkerThread		= 0;
+static volatile bool	GWorkerThreadQuit	= false;
+
+////////////////////////////////////////////////////////////////////////////////
+static void Writer_WorkerThread()
+{
+	while (!GWorkerThreadQuit)
+	{
+		const uint32 SleepMs = 100;
+		ThreadSleep(SleepMs);
+
+		Writer_UpdateControl();
+		Writer_UpdateData();
+	}
+}
+
+
+
+////////////////////////////////////////////////////////////////////////////////
+static bool GInitialized = false;
+
+////////////////////////////////////////////////////////////////////////////////
+static void Writer_LogHeader()
+{
+	UE_TRACE_EVENT_BEGIN($Trace, NewTrace)
+		UE_TRACE_EVENT_FIELD(uint16, Endian)
+		UE_TRACE_EVENT_FIELD(uint8, Version)
+		UE_TRACE_EVENT_FIELD(uint8, PointerSize)
+	UE_TRACE_EVENT_END()
+
+	UE_TRACE_LOG($Trace, NewTrace)
+		<< NewTrace.Version(1)
+		<< NewTrace.Endian(0x524d)
+		<< NewTrace.PointerSize(sizeof(void*));
+}
+
+////////////////////////////////////////////////////////////////////////////////
+static void Writer_InternalInitialize()
+{
+	if (GInitialized)
+	{
+		return;
+	}
+	GInitialized = true;
+
+	Writer_InitializeBuffers();
+	GHoldBuffer.Init();
+
+	GWorkerThread = ThreadCreate("TraceWorker", Writer_WorkerThread);
+
+	Writer_LogHeader();
+
+	Writer_InitializeControl();
+	Writer_InitializeTiming();
+	Writer_InitializeInstrumentation();
+}
+
+////////////////////////////////////////////////////////////////////////////////
+static void Writer_Shutdown()
+{
+	if (!GInitialized)
+	{
+		return;
+	}
+
+	Writer_Flush();
+
+	GWorkerThreadQuit = true;
+	ThreadJoin(GWorkerThread);
+	ThreadDestroy(GWorkerThread);
+
+	Writer_ShutdownControl();
+
+	GHoldBuffer.Shutdown();
+	Writer_ShutdownBuffers();
+
+	GInitialized = false;
+}
+
+////////////////////////////////////////////////////////////////////////////////
+void Writer_Initialize()
+{
+	using namespace Private;
+
+	if (!GInitialized)
+	{
+		static struct FInitializer
+		{
+			FInitializer()
+			{
+				Writer_InternalInitialize();
+			}
+			~FInitializer()
+			{
+				Writer_Shutdown();
+			}
+		} Initializer;
+	}
+}
+
+
+
+////////////////////////////////////////////////////////////////////////////////
+bool Writer_ConnectImpl(const ANSICHAR* Host)
+{
+	Writer_Initialize();
+
+	UPTRINT DataHandle = TcpSocketConnect(Host, 1980);
+	if (!DataHandle)
+	{
+		return false;
+	}
+
+	GPendingDataHandle = DataHandle;
+	return true;
+}
+
+////////////////////////////////////////////////////////////////////////////////
+bool Writer_ToggleEventImpl(const ANSICHAR* LoggerName, const ANSICHAR* EventName, bool State)
+{
+	Writer_Initialize();
+
+	if (FEvent* Event = FEvent::Find(LoggerName, EventName))
+	{
+		Event->bEnabled = State;
+		return true;
+	}
+
+	return false;
+}
+
+} // namespace Private
+} // namespace Trace
+
+#endif // UE_TRACE_ENABLED
diff --git a/Engine/Source/Runtime/TraceLog/Public/Trace/Private/Event.h b/Engine/Source/Runtime/TraceLog/Public/Trace/Private/Event.h
new file mode 100644
--- /dev/null
+++ b/Engine/Source/Runtime/TraceLog/Public/Trace/Private/Event.h
@@ -0,0 +1,78 @@
+//
+// Copyright 1998-2019 Epic Games, Inc. All Rights Reserved.
+//
+
+#pragma once
+
+#include "Trace.h"
+
+namespace Trace
+{
+
+////////////////////////////////////////////////////////////////////////////////
+enum class EKnownEventUids : uint16
+{
+	NewEvent,
+	User,
+	Max			= 1 << 14, // ...leaves two MSB bits for other uses.
+};
+
+////////////////////////////////////////////////////////////////////////////////
+struct FNewEventEvent
+{
+	uint16		Uid;
+	uint16		FieldCount;
+	uint8		LoggerNameSize;
+	uint8		EventNameSize;
+	struct
+	{
+		uint16	Offset;
+		uint16	Size;
+		uint8	TypeInfo;
+		uint8	NameSize;
+	}			Fields[];
+	/*uint8		NameData[]*/
+};
+
+} // namespace Trace
+
+
+
+#if UE_TRACE_ENABLED
+
+namespace Trace
+{
+
+struct FFieldDesc;
+struct FLiteralName;
+
+////////////////////////////////////////////////////////////////////////////////
+class FEvent
+{
+public:
+	enum { HeaderSize = sizeof(uint32), }; // I'll just dump this here for now.
+
+	class FLogScope
+	{
+	public:
+						FLogScope(uint16 EventUid, uint16 Size);
+						FLogScope(uint16 EventUid, uint16 Size, uint16 ExtraBytes);
+						FLogScope(uint16 EventUid, uint16 Size, uint8* Out);
+						~FLogScope();
+		uint8*			Ptr;
+		const bool		bOutOfBand = false;
+	};
+
+	void*						Handle;
+	uint32						LoggerHash;
+	uint32						Hash;
+	uint16						Uid;
+	bool						bInitialized;
+	bool						bEnabled;
+	UE_TRACE_API static void	Create(FEvent* Target, const FLiteralName& LoggerName, const FLiteralName& EventName, const FFieldDesc* FieldDescs, uint32 FieldCount);
+	static FEvent*				Find(const ANSICHAR* LoggerName, const ANSICHAR* EventName);
+};
+
+} // namespace Trace
+
+#endif // UE_TRACE_ENABLED
diff --git a/Engine/Source/Runtime/TraceLog/Public/Trace/Private/Event.inl b/Engine/Source/Runtime/TraceLog/Public/Trace/Private/Event.inl
new file mode 100644
--- /dev/null
+++ b/Engine/Source/Runtime/TraceLog/Public/Trace/Private/Event.inl
@@ -0,0 +1,54 @@
+//
+// Copyright 1998-2019 Epic Games, Inc. All Rights Reserved.
+//
+
+#pragma once
+
+#if UE_TRACE_ENABLED
+
+#include "Event.h"
+#include "Writer.inl"
+
+namespace Trace
+{
+
+////////////////////////////////////////////////////////////////////////////////
+inline FEvent::FLogScope::FLogScope(uint16 EventUid, uint16 Size)
+{
+	Ptr = Writer_BeginLog(EventUid, Size);
+}
+
+////////////////////////////////////////////////////////////////////////////////
+inline FEvent::FLogScope::FLogScope(uint16 EventUid, uint16 Size, uint16 ExtraBytes)
+{
+	Ptr = Writer_BeginLog(EventUid, Size + ExtraBytes);
+}
+
+////////////////////////////////////////////////////////////////////////////////
+inline FEvent::FLogScope::FLogScope(uint16 EventUid, uint16 Size, uint8* Out)
+: bOutOfBand(true)
+{
+	*(uint32*)Out = (uint32(Size) << 16)|uint32(EventUid);
+	Ptr = Out + HeaderSize;
+}
+
+////////////////////////////////////////////////////////////////////////////////
+inline FEvent::FLogScope::~FLogScope()
+{
+	if (!bOutOfBand)
+		Writer_EndLog(Ptr);
+}
+
+
+
+////////////////////////////////////////////////////////////////////////////////
+template <typename ActionType>
+inline const FEvent::FLogScope& operator << (const FEvent::FLogScope& Lhs, const ActionType& Rhs)
+{
+	Rhs.Write(Lhs.Ptr);
+	return Lhs;
+}
+
+} // namespace Trace
+
+#endif // UE_TRACE_ENABLED
diff --git a/Engine/Source/Runtime/TraceLog/Public/Trace/Private/Field.h b/Engine/Source/Runtime/TraceLog/Public/Trace/Private/Field.h
new file mode 100644
--- /dev/null
+++ b/Engine/Source/Runtime/TraceLog/Public/Trace/Private/Field.h
@@ -0,0 +1,221 @@
+//
+// Copyright 1998-2019 Epic Games, Inc. All Rights Reserved.
+//
+
+#pragma once
+
+#include "Trace.h"
+
+namespace Trace
+{
+
+////////////////////////////////////////////////////////////////////////////////
+enum : uint8
+{
+	/* Category */
+	_Field_Integer		= 0000, 
+	_Field_Float		= 0100,
+	_Field_Array		= 0200, /* flag */
+
+	/* Size */
+	_Field_Pow2SizeMask	= 0003,
+	_Field_8			= 0000,
+	_Field_16			= 0001,
+	_Field_32			= 0002,
+	_Field_64			= 0003,
+#if PLATFORM_64BITS
+	_Field_Ptr			= _Field_64,
+#else
+	_Field_Ptr			= _Field_32,
+#endif
+
+	/* unused			= 0004, flag */
+	/* unused			= 0070, enum/flags*/
+};
+
+////////////////////////////////////////////////////////////////////////////////
+enum EFieldType : uint8
+{
+	Field_Bool		= _Field_Integer | _Field_8,
+	Field_Int8		= _Field_Integer | _Field_8,
+	Field_Int16		= _Field_Integer | _Field_16,
+	Field_Int32		= _Field_Integer | _Field_32,
+	Field_Int64		= _Field_Integer | _Field_64,
+	Field_Ptr		= _Field_Integer | _Field_Ptr,
+	Field_Float		= _Field_Float	 | _Field_32,
+	Field_Double	= _Field_Float	 | _Field_64,
+
+	/* Flags */
+	Field_Array		= _Field_Array,
+};
+
+////////////////////////////////////////////////////////////////////////////////
+template <typename Type> struct TFieldType;
+
+template <> struct TFieldType<bool>			{ enum { Value = Field_Bool }; };
+template <> struct TFieldType<int8>			{ enum { Value = Field_Int8 }; };
+template <> struct TFieldType<int16>		{ enum { Value = Field_Int16 }; };
+template <> struct TFieldType<int32>		{ enum { Value = Field_Int32 }; };
+template <> struct TFieldType<int64>		{ enum { Value = Field_Int64 }; };
+template <> struct TFieldType<uint8>		{ enum { Value = Field_Int8 }; };
+template <> struct TFieldType<uint16>		{ enum { Value = Field_Int16 }; };
+template <> struct TFieldType<uint32>		{ enum { Value = Field_Int32 }; };
+template <> struct TFieldType<uint64>		{ enum { Value = Field_Int64 }; };
+template <> struct TFieldType<float>		{ enum { Value = Field_Float }; };
+template <> struct TFieldType<double>		{ enum { Value = Field_Double }; };
+template <class T> struct TFieldType<T*>	{ enum { Value = Field_Ptr }; };
+
+} // namespace Trace
+
+#if UE_TRACE_ENABLED
+
+#include "Templates/UnrealTemplate.h"
+
+namespace Trace
+{
+
+////////////////////////////////////////////////////////////////////////////////
+struct FLiteralName
+{
+	template <uint32 Size>
+	explicit FLiteralName(const ANSICHAR (&Name)[Size])
+	: Ptr(Name)
+	, Length(Size - 1)
+	{
+		static_assert(Size < 256, "Field name is too large");
+	}
+
+	const ANSICHAR* Ptr;
+	uint8 Length;
+};
+
+////////////////////////////////////////////////////////////////////////////////
+struct FFieldDesc
+{
+	FFieldDesc(const FLiteralName& Name, uint8 Type, uint16 Offset, uint16 Size)
+	: Name(Name.Ptr)
+	, ValueOffset(Offset)
+	, ValueSize(Size)
+	, NameSize(Name.Length)
+	, TypeInfo(Type)
+	{
+	}
+
+	const ANSICHAR* Name;
+	uint16			ValueOffset;
+	uint16			ValueSize;
+	uint8			NameSize;
+	uint8			TypeInfo;
+};
+
+////////////////////////////////////////////////////////////////////////////////
+template <int OffsetValue, typename Type>
+struct TFieldBase
+	: public FFieldDesc
+{
+	enum { Offset = OffsetValue, Size = sizeof(Type) };
+
+protected:
+	TFieldBase(const FLiteralName& Name)
+	: FFieldDesc(Name, TFieldType<Type>::Value, OffsetValue, Size)
+	{
+	}
+};
+
+////////////////////////////////////////////////////////////////////////////////
+template <int Offset, typename Type> struct TField;
+
+////////////////////////////////////////////////////////////////////////////////
+// Used to terminate the field list and determine an event's size.
+enum EndOfFields {};
+template <int Size>
+struct TField<Size, EndOfFields>
+{
+	enum : uint16 { Value = Size };
+};
+
+////////////////////////////////////////////////////////////////////////////////
+// Access to additional data that can be included along with a logged event.
+enum Attachment {};
+template <int Offset>
+struct TField<Offset, Attachment>
+{
+	template <typename LambdaType>
+	struct FActionableLambda
+	{
+		LambdaType& Value;
+		void Write(uint8* __restrict Ptr) const
+		{
+			Value(Ptr + Offset);
+		}
+	};
+
+	template <typename LambdaType>
+	const FActionableLambda<LambdaType> operator () (LambdaType&& Lambda) const
+	{
+		return { Forward<LambdaType>(Lambda) };
+	}
+
+
+	struct FActionableMemcpy
+	{
+		const void* Data;
+		uint32 Size;
+		void Write(uint8* __restrict Ptr) const
+		{
+			::memcpy(Ptr + Offset, Data, Size);
+		}
+	};
+
+	const FActionableMemcpy operator () (const void* Data, uint32 Size) const
+	{
+		return { Data, Size };
+	}
+};
+
+#if 0
+////////////////////////////////////////////////////////////////////////////////
+template <typename Type>
+struct TField<Type[]>
+	: public TFieldBase<sizeof(uint32)>
+{
+	TField(const ANSICHAR (Name)[]) : TFieldBase<sizeof(uint32)>(Name) {}
+};
+
+////////////////////////////////////////////////////////////////////////////////
+template <int Offset, typename Type, int Count>
+struct TField<Offset, Type[Count]>
+	: public TFieldBase<Offset, Type>
+{
+	TField(const ANSICHAR (Name)[]) : TFieldBase<Offset, Type[Count]>(Name) {}
+};
+#endif
+
+////////////////////////////////////////////////////////////////////////////////
+template <int Offset, typename Type>
+struct TField
+	: public TFieldBase<Offset, Type>
+{
+	struct FActionable
+	{
+		Type Value;
+		void Write(uint8* __restrict Ptr) const
+		{
+			*(Type*)(Ptr + Offset) = Value;
+		}
+	};
+
+	TField(const FLiteralName& Name)
+	: TFieldBase<Offset, Type>(Name)
+	{
+	}
+
+	const FActionable operator () (const Type& __restrict Value) const
+	{
+		return { Value };
+	}
+};
+
+} // namespace Trace
+
+#endif // UE_TRACE_ENABLED
diff --git a/Engine/Source/Runtime/TraceLog/Public/Trace/Private/Trace.h b/Engine/Source/Runtime/TraceLog/Public/Trace/Private/Trace.h
new file mode 100644
--- /dev/null
+++ b/Engine/Source/Runtime/TraceLog/Public/Trace/Private/Trace.h
@@ -0,0 +1,27 @@
+//
+// Copyright 1998-2019 Epic Games, Inc. All Rights Reserved.
+//
+
+#pragma once
+
+#include "CoreTypes.h"
+
+////////////////////////////////////////////////////////////////////////////////
+#if !IS_PROGRAM && !defined(UE_TRACE_ENABLED)
+#	if PLATFORM_WINDOWS || PLATFORM_PS4 || PLATFORM_XBOXONE
+#		define UE_TRACE_ENABLED	1
+#	else
+#	endif
+#endif // !IS_PROGRAM
+
+#if !defined(UE_TRACE_ENABLED)
+#	define UE_TRACE_ENABLED		0
+#endif
+
+#if UE_TRACE_ENABLED
+#	define UE_TRACE_IMPL(...)
+#	define UE_TRACE_API			TRACELOG_API
+#else
+#	define UE_TRACE_IMPL(...)	{ return __VA_ARGS__; }
+#	define UE_TRACE_API			inline
+#endif
diff --git a/Engine/Source/Runtime/TraceLog/Public/Trace/Private/Writer.inl b/Engine/Source/Runtime/TraceLog/Public/Trace/Private/Writer.inl
new file mode 100644
--- /dev/null
+++ b/Engine/Source/Runtime/TraceLog/Public/Trace/Private/Writer.inl
@@ -0,0 +1,92 @@
+//
+// Copyright 1998-2019 Epic Games, Inc. All Rights Reserved.
+//
+
+#pragma once
+
+#if UE_TRACE_ENABLED
+
+#include <atomic>
+
+namespace Trace
+{
+
+////////////////////////////////////////////////////////////////////////////////
+template <typename Type>
+using TTraceAtomic = std::atomic<Type>;
+
+namespace Private
+{
+
+////////////////////////////////////////////////////////////////////////////////
+struct FBuffer
+{
+	FBuffer()
+	: Next(nullptr)
+	, Used(sizeof(FBuffer))
+	{
+	}
+
+	union
+	{
+		TTraceAtomic<FBuffer*>	Next;
+		char					Padding0[PLATFORM_CACHE_LINE_SIZE];
+	};
+	union
+	{
+		TTraceAtomic<uint32>	Used;
+		char					Padding1[PLATFORM_CACHE_LINE_SIZE];
+	};
+	uint32						Final;
+	alignas(void*) uint8		Data[];
+};
+
+extern UE_TRACE_API TTraceAtomic<FBuffer*> GActiveBuffer;
+
+static const uint16		BufferSizePow2	= 19;
+static const uint32		BufferSize		= 1 << BufferSizePow2;
+static const uint32		BufferSizeMask	= BufferSize - 1;
+static const uint32		BufferRefBit	= BufferSize << 1;
+
+extern UE_TRACE_API void* Writer_NextBuffer(FBuffer*, uint32, uint32);
+
+} // Private
+
+////////////////////////////////////////////////////////////////////////////////
+inline uint8* Writer_BeginLog(uint16 EventUid, uint16 Size)
+{
+	using namespace Private;
+
+	uint32 AllocSize = Size;
+	AllocSize += sizeof(uint32);
+	AllocSize += BufferRefBit;
+
+	// Fetch buffer and claim some space in it.
+	FBuffer* Buffer = GActiveBuffer.load(std::memory_order_acquire);
+	uint32 PrevUsed = Buffer->Used.fetch_add(AllocSize, std::memory_order_relaxed);
+	uint32 Used = PrevUsed + AllocSize;
+
+	uint32* Out = (uint32*)(UPTRINT(Buffer) + (BufferSizeMask & PrevUsed));
+
+	// Did we exhaust the active buffer?
+	if (UNLIKELY(Used & BufferSize))
+	{
+		Out = (uint32*)Writer_NextBuffer(Buffer, PrevUsed, AllocSize);
+		Buffer = (FBuffer*)(UPTRINT(Out) & ~UPTRINT(BufferSizeMask));
+	}
+
+	Out[0] = (uint32(Size) << 16)|uint32(EventUid);
+	return (uint8*)(Out + 1);
+}
+
+////////////////////////////////////////////////////////////////////////////////
+inline void Writer_EndLog(uint8* EventData)
+{
+	using namespace Private;
+	auto* Buffer = (FBuffer*)(UPTRINT(EventData) & ~UPTRINT(BufferSizeMask));
+	Buffer->Used.fetch_sub(BufferRefBit, std::memory_order_release);
+}
+
+} // namespace Trace
+
+#endif // UE_TRACE_ENABLED
diff --git a/Engine/Source/Runtime/TraceLog/Public/Trace/Trace.h b/Engine/Source/Runtime/TraceLog/Public/Trace/Trace.h
new file mode 100644
--- /dev/null
+++ b/Engine/Source/Runtime/TraceLog/Public/Trace/Trace.h
@@ -0,0 +1,108 @@
+//
+// Copyright 1998-2019 Epic Games, Inc. All Rights Reserved.
+//
+
+#pragma once
+
+#include "Private/Trace.h"
+
+////////////////////////////////////////////////////////////////////////////////
+namespace Trace
+{
+
+UE_TRACE_API bool Connect(const TCHAR* Host) UE_TRACE_IMPL(false);
+UE_TRACE_API bool ToggleEvent(const TCHAR* LoggerName, const TCHAR* EventName, bool State) UE_TRACE_IMPL(false);
+
+} // namespace Trace
+
+#if UE_TRACE_ENABLED
+
+#define UE_TRACE_EVENT_DEFINE(LoggerName, EventName) \
+	Trace::FEvent LoggerName##EventName##Event;
+
+#define UE_TRACE_EVENT_BEGIN(LoggerName, EventName) \
+	UE_TRACE_EVENT_BEGIN_IMPL(static, LoggerName, EventName)
+
+#define UE_TRACE_EVENT_BEGIN_EXTERN(LoggerName, EventName) \
+	UE_TRACE_EVENT_BEGIN_IMPL(extern, LoggerName, EventName)
+
+#define UE_TRACE_EVENT_BEGIN_IMPL(LinkageType, LoggerName, EventName) \
+	LinkageType UE_TRACE_EVENT_DEFINE(LoggerName, EventName) \
+	struct F##LoggerName##EventName##Fields \
+	{ \
+		static void FORCENOINLINE Initialize() \
+		{ \
+			static const bool bOnceOnly = [] () { \
+				F##LoggerName##EventName##Fields Fields; \
+				const auto* Descs = (Trace::FFieldDesc*)(&Fields); \
+				uint32 DescCount = uint32(sizeof(Fields) / sizeof(*Descs)); \
+				const auto& LoggerLiteral = Trace::FLiteralName(#LoggerName); \
+				const auto& EventLiteral = Trace::FLiteralName(#EventName); \
+				Trace::FEvent::Create(&LoggerName##EventName##Event, LoggerLiteral, EventLiteral, Descs, DescCount); \
+				return true; \
+			}(); \
+		} \
+		Trace::TField<0, 
+
+#define UE_TRACE_EVENT_FIELD(FieldType, FieldName) \
+		FieldType> const FieldName = Trace::FLiteralName(#FieldName); \
+		Trace::TField<decltype(FieldName)::Offset + decltype(FieldName)::Size,
+
+#define UE_TRACE_EVENT_END() \
+		Trace::EndOfFields> const EventSize_Private = {}; \
+		Trace::TField<decltype(EventSize_Private)::Value, Trace::Attachment> const Attachment = {}; \
+		explicit operator bool () const { return true; } \
+	};
+
+#define UE_TRACE_EVENT_SIZE(LoggerName, EventName) \
+	decltype(F##LoggerName##EventName##Fields::EventSize_Private)::Value
+
+#define UE_TRACE_LOG(LoggerName, EventName, ...) \
+	if (!LoggerName##EventName##Event.bInitialized) \
+		F##LoggerName##EventName##Fields::Initialize(); \
+	if (LoggerName##EventName##Event.bEnabled) \
+		if (const auto& __restrict EventName = (F##LoggerName##EventName##Fields&)LoggerName##EventName##Event) \
+			Trace::FEvent::FLogScope(LoggerName##EventName##Event.Uid, UE_TRACE_EVENT_SIZE(LoggerName, EventName), ##__VA_ARGS__)
+
+#else // UE_TRACE_ENABLED
+
+#define UE_TRACE_EVENT_DEFINE(LoggerName, EventName)
+
+#define UE_TRACE_EVENT_BEGIN(LoggerName, EventName) \
+	UE_TRACE_EVENT_BEGIN_IMPL(LoggerName, EventName)
+
+#define UE_TRACE_EVENT_BEGIN_EXTERN(LoggerName, EventName) \
+	UE_TRACE_EVENT_BEGIN_IMPL(LoggerName, EventName)
+
+#define UE_TRACE_EVENT_BEGIN_IMPL(LoggerName, EventName) \
+	struct F##LoggerName##EventName##Dummy \
+	{ \
+		struct FTraceDisabled \
+		{ \
+			const FTraceDisabled& operator () (...) const { return *this; } \
+		}; \
+		const F##LoggerName##EventName##Dummy& operator << (const FTraceDisabled&) const \
+		{ \
+			return *this; \
+		} \
+		explicit operator bool () const { return false; }
+
+#define UE_TRACE_EVENT_FIELD(FieldType, FieldName) \
+		const FTraceDisabled& FieldName;
+
+#define UE_TRACE_EVENT_END() \
+		const FTraceDisabled& Attachment; \
+	};
+
+#define UE_TRACE_EVENT_SIZE(EventName) \
+	1
+
+#define UE_TRACE_LOG(LoggerName, EventName, ...) \
+	if (const auto& EventName = *(F##LoggerName##EventName##Dummy*)1) \
+		EventName
+
+#endif // UE_TRACE_ENABLED
+
+#include "Private/Field.h"
+#include "Private/Event.inl"
+#include "Private/Writer.inl"
diff --git a/Engine/Source/Runtime/TraceLog/TraceLog.Build.cs b/Engine/Source/Runtime/TraceLog/TraceLog.Build.cs
new file mode 100644
--- /dev/null
+++ b/Engine/Source/Runtime/TraceLog/TraceLog.Build.cs
@@ -0,0 +1,19 @@
+// Copyright 1998-2019 Epic Games, Inc. All Rights Reserved.
+
+using UnrealBuildTool;
+using System;
+using System.IO;
+
+public class TraceLog : ModuleRules
+{
+	public TraceLog(ReadOnlyTargetRules Target) : base(Target)
+	{
+		bRequiresImplementModule = false;
+		PrivateIncludePathModuleNames.Add("Core");
+
+        if (Target.Platform == UnrealTargetPlatform.PS4)
+        {
+            PublicAdditionalLibraries.Add("SceNet_stub_weak");
+        }
+    }
+}
