Fixed PVS Studio warning.

diff --git a/Engine/Source/Runtime/TraceLog/Private/Trace/Windows/Detail.cpp b/Engine/Source/Runtime/TraceLog/Private/Trace/Windows/Detail.cpp
--- a/Engine/Source/Runtime/TraceLog/Private/Trace/Windows/Detail.cpp
+++ b/Engine/Source/Runtime/TraceLog/Private/Trace/Windows/Detail.cpp
@@ -36,7 +36,8 @@ void MemoryFree(void* Address, SIZE_T Size)
 ////////////////////////////////////////////////////////////////////////////////
 void MemoryMap(void* Address, SIZE_T Size)
 {
-	VirtualAlloc(Address, Size, MEM_COMMIT, PAGE_READWRITE);
+	void* Ret = VirtualAlloc(Address, Size, MEM_COMMIT, PAGE_READWRITE);
+	(void)Ret;
 }
 
 ////////////////////////////////////////////////////////////////////////////////
