Linux compile warning.

diff --git a/Engine/Source/Runtime/TraceLog/Private/Trace/Writer.cpp b/Engine/Source/Runtime/TraceLog/Private/Trace/Writer.cpp
--- a/Engine/Source/Runtime/TraceLog/Private/Trace/Writer.cpp
+++ b/Engine/Source/Runtime/TraceLog/Private/Trace/Writer.cpp
@@ -6,6 +6,8 @@
 
 #include "Trace/Trace.h"
 
+#include <emmintrin.h>
+
 namespace Trace
 {
 
