Temporary workaround to unblock iOS builds. It appears the build machines don't have the same set of libraries available to them as the most recent XCode versions on our development machines

diff --git a/Engine/Source/Runtime/TraceLog/Public/Trace/Detail/Trace.h b/Engine/Source/Runtime/TraceLog/Public/Trace/Detail/Trace.h
--- a/Engine/Source/Runtime/TraceLog/Public/Trace/Detail/Trace.h
+++ b/Engine/Source/Runtime/TraceLog/Public/Trace/Detail/Trace.h
@@ -7,7 +7,7 @@
 ////////////////////////////////////////////////////////////////////////////////
 #if !defined(UE_TRACE_ENABLED)
 #	if !UE_BUILD_SHIPPING && !IS_PROGRAM
-#		if PLATFORM_WINDOWS || PLATFORM_APPLE || PLATFORM_PS4 || PLATFORM_XBOXONE || PLATFORM_SWITCH || PLATFORM_ANDROID
+#		if PLATFORM_WINDOWS || /* PLATFORM_APPLE || */ PLATFORM_PS4 || PLATFORM_XBOXONE || PLATFORM_SWITCH || PLATFORM_ANDROID
 #			define UE_TRACE_ENABLED	1
 #		endif
 #	endif
