Removed unused CoreServices.h include.

#rnx

(ushell-p4-cherrypick of 9470131 by [REDACTED_PERSON])

diff --git a/Engine/Source/Runtime/TraceLog/Private/Trace/Detail/Apple/AppleTrace.cpp b/Engine/Source/Runtime/TraceLog/Private/Trace/Detail/Apple/AppleTrace.cpp
--- a/Engine/Source/Runtime/TraceLog/Private/Trace/Detail/Apple/AppleTrace.cpp
+++ b/Engine/Source/Runtime/TraceLog/Private/Trace/Detail/Apple/AppleTrace.cpp
@@ -4,7 +4,6 @@
 
 #if UE_TRACE_ENABLED
 
-#include <CoreServices/CoreServices.h>
 #include <fcntl.h>
 #include <mach/mach.h>
 #include <mach/mach_time.h>
