Fixed TimeGetFrequency() calculation on Apple platforms so that it returns the number of mach_absolute_time() samples in one second.

#rnx

(ushell-p4-cherrypick of 9485845 by [REDACTED_PERSON])

diff --git a/Engine/Source/Runtime/TraceLog/Private/Trace/Detail/Apple/AppleTrace.cpp b/Engine/Source/Runtime/TraceLog/Private/Trace/Detail/Apple/AppleTrace.cpp
--- a/Engine/Source/Runtime/TraceLog/Private/Trace/Detail/Apple/AppleTrace.cpp
+++ b/Engine/Source/Runtime/TraceLog/Private/Trace/Detail/Apple/AppleTrace.cpp
@@ -99,7 +99,7 @@ const mach_timebase_info_data_t& TimeGetInfo()
 uint64 TimeGetFrequency()
 {
 	const mach_timebase_info_data_t& Info = TimeGetInfo();
-	return (uint64)((double)Info.numer / (1e-9 * (double)Info.denom) + 0.5);
+	return (uint64(1 * 1000 * 1000 * 1000) * uint64(Info.denom)) / uint64(Info.numer);
 }
 
 ////////////////////////////////////////////////////////////////////////////////
