Fixed merge errors when Dev-RenPlatStage copied up to Main.

#rnx

(ushell-p4-cherrypick of 9631298 by [REDACTED_PERSON])
(ushell-p4-cherrypick of 9631305 by [REDACTED_PERSON])
(ushell-p4-cherrypick of 9631347 by [REDACTED_PERSON])

diff --git a/Engine/Source/Runtime/Core/Core.Build.cs b/Engine/Source/Runtime/Core/Core.Build.cs
--- a/Engine/Source/Runtime/Core/Core.Build.cs
+++ b/Engine/Source/Runtime/Core/Core.Build.cs
@@ -15,6 +15,7 @@ public class Core : ModuleRules
 		PrivateDependencyModuleNames.Add("BuildSettings");
 
 		PublicDependencyModuleNames.Add("TraceLog");
+		PublicIncludePaths.Add("Runtime/TraceLog/Public");
 
 		PrivateIncludePaths.AddRange(
 			new string[] {
diff --git a/Engine/Source/Runtime/Core/Private/Stats/StatsTrace.cpp b/Engine/Source/Runtime/Core/Private/Stats/StatsTrace.cpp
--- a/Engine/Source/Runtime/Core/Private/Stats/StatsTrace.cpp
+++ b/Engine/Source/Runtime/Core/Private/Stats/StatsTrace.cpp
@@ -5,6 +5,7 @@
 #include "Misc/CString.h"
 #include "HAL/PlatformTime.h"
 #include "HAL/PlatformTLS.h"
+#include "Trace/Trace.h"
 #include "UObject/NameTypes.h"
 
 #if STATSTRACE_ENABLED
diff --git a/Engine/Source/Runtime/Core/Public/Logging/LogTrace.h b/Engine/Source/Runtime/Core/Public/Logging/LogTrace.h
--- a/Engine/Source/Runtime/Core/Public/Logging/LogTrace.h
+++ b/Engine/Source/Runtime/Core/Public/Logging/LogTrace.h
@@ -3,9 +3,10 @@
 #pragma once
 
 #include "CoreTypes.h"
+#include "Trace/Config.h"
 
 #if !defined(LOGTRACE_ENABLED)
-#if !IS_PROGRAM && !UE_BUILD_SHIPPING && (PLATFORM_WINDOWS || PLATFORM_PS4 || PLATFORM_UNIX || PLATFORM_XBOXONE)
+#if UE_TRACE_ENABLED && !UE_BUILD_SHIPPING
 #define LOGTRACE_ENABLED 1
 #else
 #define LOGTRACE_ENABLED 0
diff --git a/Engine/Source/Runtime/Core/Public/ProfilingDebugging/CpuProfilerTrace.h b/Engine/Source/Runtime/Core/Public/ProfilingDebugging/CpuProfilerTrace.h
--- a/Engine/Source/Runtime/Core/Public/ProfilingDebugging/CpuProfilerTrace.h
+++ b/Engine/Source/Runtime/Core/Public/ProfilingDebugging/CpuProfilerTrace.h
@@ -3,9 +3,10 @@
 #pragma once
 
 #include "CoreTypes.h"
+#include "Trace/Config.h"
 
 #if !defined(CPUPROFILERTRACE_ENABLED)
-#if !IS_PROGRAM && !UE_BUILD_SHIPPING && (PLATFORM_WINDOWS || PLATFORM_UNIX || PLATFORM_PS4 || PLATFORM_XBOXONE || PLATFORM_SWITCH || PLATFORM_APPLE || PLATFORM_ANDROID)
+#if UE_TRACE_ENABLED && !UE_BUILD_SHIPPING
 #define CPUPROFILERTRACE_ENABLED 1
 #else
 #define CPUPROFILERTRACE_ENABLED 0
diff --git a/Engine/Source/Runtime/Core/Public/ProfilingDebugging/MiscTrace.h b/Engine/Source/Runtime/Core/Public/ProfilingDebugging/MiscTrace.h
--- a/Engine/Source/Runtime/Core/Public/ProfilingDebugging/MiscTrace.h
+++ b/Engine/Source/Runtime/Core/Public/ProfilingDebugging/MiscTrace.h
@@ -3,7 +3,7 @@
 #pragma once
 
 #include "CoreTypes.h"
-#include "Trace/Trace.h"
+#include "Trace/Config.h"
 #include "ProfilingDebugging/FormatArgsTrace.h"
 
 #if UE_TRACE_ENABLED && !UE_BUILD_SHIPPING
diff --git a/Engine/Source/Runtime/Core/Public/Stats/StatsTrace.h b/Engine/Source/Runtime/Core/Public/Stats/StatsTrace.h
--- a/Engine/Source/Runtime/Core/Public/Stats/StatsTrace.h
+++ b/Engine/Source/Runtime/Core/Public/Stats/StatsTrace.h
@@ -3,7 +3,7 @@
 #pragma once
 
 #include "CoreTypes.h"
-#include "Trace/Trace.h"
+#include "Trace/Config.h"
 
 #define EXPERIMENTAL_STATSTRACE_ENABLED 0
 
diff --git a/Engine/Source/Runtime/TraceLog/Private/Trace/Mac/Detail.cpp b/Engine/Source/Runtime/TraceLog/Private/Trace/Mac/Detail.cpp
deleted file mode 100644
diff --git a/Engine/Source/Runtime/TraceLog/Public/Trace/Config.h b/Engine/Source/Runtime/TraceLog/Public/Trace/Config.h
new file mode 100644
--- /dev/null
+++ b/Engine/Source/Runtime/TraceLog/Public/Trace/Config.h
@@ -0,0 +1,17 @@
+// Copyright 1998-2019 Epic Games, Inc. All Rights Reserved.
+
+#pragma once
+
+#include "HAL/Platform.h"
+
+#if !defined(UE_TRACE_ENABLED)
+#	if !UE_BUILD_SHIPPING && !IS_PROGRAM
+#		if PLATFORM_WINDOWS || PLATFORM_UNIX || PLATFORM_APPLE || PLATFORM_PS4 || PLATFORM_XBOXONE || PLATFORM_SWITCH || PLATFORM_ANDROID
+#			define UE_TRACE_ENABLED	1
+#		endif
+#	endif
+#endif // !IS_PROGRAM
+
+#if !defined(UE_TRACE_ENABLED)
+#	define UE_TRACE_ENABLED 0
+#endif
diff --git a/Engine/Source/Runtime/TraceLog/Public/Trace/Detail/Trace.h b/Engine/Source/Runtime/TraceLog/Public/Trace/Detail/Trace.h
--- a/Engine/Source/Runtime/TraceLog/Public/Trace/Detail/Trace.h
+++ b/Engine/Source/Runtime/TraceLog/Public/Trace/Detail/Trace.h
@@ -3,20 +3,9 @@
 #pragma once
 
 #include "CoreTypes.h"
+#include "Trace/Config.h"
 
 ////////////////////////////////////////////////////////////////////////////////
-#if !defined(UE_TRACE_ENABLED)
-#	if !UE_BUILD_SHIPPING && !IS_PROGRAM
-#		if PLATFORM_WINDOWS || PLATFORM_UNIX || PLATFORM_APPLE || PLATFORM_PS4 || PLATFORM_XBOXONE || PLATFORM_SWITCH || PLATFORM_ANDROID
-#			define UE_TRACE_ENABLED	1
-#		endif
-#	endif
-#endif // !IS_PROGRAM
-
-#if !defined(UE_TRACE_ENABLED)
-#	define UE_TRACE_ENABLED		0
-#endif
-
 #if UE_TRACE_ENABLED
 #	define UE_TRACE_IMPL(...)
 #	define UE_TRACE_API			TRACELOG_API
