Added a handful of missing includes on Android.

#rnx

(ushell-p4-cherrypick of 9730430 by [REDACTED_PERSON])

diff --git a/Engine/Source/Runtime/TraceLog/Private/Trace/Detail/Android/AndroidTrace.cpp b/Engine/Source/Runtime/TraceLog/Private/Trace/Detail/Android/AndroidTrace.cpp
--- a/Engine/Source/Runtime/TraceLog/Private/Trace/Detail/Android/AndroidTrace.cpp
+++ b/Engine/Source/Runtime/TraceLog/Private/Trace/Detail/Android/AndroidTrace.cpp
@@ -5,8 +5,12 @@
 #if UE_TRACE_ENABLED
 
 #include <arpa/inet.h>
+#include <fcntl.h>
+#include <pthread.h>
 #include <sys/mman.h>
 #include <sys/socket.h>
+#include <sys/stat.h>
+#include <unistd.h>
 
 namespace Trace {
 namespace Private {
