Added missing include in AppleTrace.cpp

#rnx

(ushell-p4-cherrypick of 9730388 by [REDACTED_PERSON])

diff --git a/Engine/Source/Runtime/TraceLog/Private/Trace/Detail/Apple/AppleTrace.cpp b/Engine/Source/Runtime/TraceLog/Private/Trace/Detail/Apple/AppleTrace.cpp
--- a/Engine/Source/Runtime/TraceLog/Private/Trace/Detail/Apple/AppleTrace.cpp
+++ b/Engine/Source/Runtime/TraceLog/Private/Trace/Detail/Apple/AppleTrace.cpp
@@ -4,6 +4,7 @@
 
 #if UE_TRACE_ENABLED
 
+#include <errno.h>
 #include <fcntl.h>
 #include <mach/mach.h>
 #include <mach/mach_time.h>
