Skip trying to write an event that is too large

diff --git a/Engine/Source/Runtime/TraceLog/Private/Trace/Writer.cpp b/Engine/Source/Runtime/TraceLog/Private/Trace/Writer.cpp
--- a/Engine/Source/Runtime/TraceLog/Private/Trace/Writer.cpp
+++ b/Engine/Source/Runtime/TraceLog/Private/Trace/Writer.cpp
@@ -158,6 +158,12 @@ static FWriteBuffer* Writer_NextBufferInternal(uint32 PageGrowth)
 ////////////////////////////////////////////////////////////////////////////////
 UE_TRACE_API uint8* Writer_NextBuffer(uint16 Size)
 {
+	if (Size >= GPoolBlockSize - sizeof(FWriteBuffer))
+	{
+		/* Someone is trying to write an event that is too large */
+		return nullptr;
+	}
+
 	FWriteBuffer* Current = GWriteBuffer;
 
 	// Carry along or assign a new thread id
