A further attempt to placate PVS.

diff --git a/Engine/Source/Runtime/TraceLog/Private/Trace/Windows/Detail.cpp b/Engine/Source/Runtime/TraceLog/Private/Trace/Windows/Detail.cpp
--- a/Engine/Source/Runtime/TraceLog/Private/Trace/Windows/Detail.cpp
+++ b/Engine/Source/Runtime/TraceLog/Private/Trace/Windows/Detail.cpp
@@ -34,8 +34,11 @@ void MemoryFree(void* Address, SIZE_T Size)
 ////////////////////////////////////////////////////////////////////////////////
 void MemoryMap(void* Address, SIZE_T Size)
 {
-	void* Ret = VirtualAlloc(Address, Size, MEM_COMMIT, PAGE_READWRITE);
-	(void)Ret;
+	auto Inner = [] (void* Address, SIZE_T Size) -> void*
+	{
+		return VirtualAlloc(Address, Size, MEM_COMMIT, PAGE_READWRITE);
+	};
+	Inner(Address, Size);
 }
 
 ////////////////////////////////////////////////////////////////////////////////
