Fixed events going missing from the end of a trace log.

diff --git a/Engine/Source/Runtime/Launch/Private/LaunchEngineLoop.cpp b/Engine/Source/Runtime/Launch/Private/LaunchEngineLoop.cpp
--- a/Engine/Source/Runtime/Launch/Private/LaunchEngineLoop.cpp
+++ b/Engine/Source/Runtime/Launch/Private/LaunchEngineLoop.cpp
@@ -5053,6 +5053,8 @@ void FEngineLoop::AppPreExit( )
 		GShaderCompilingManager = nullptr;
 	}
 #endif
+
+	Trace::Flush();
 }
 
 
diff --git a/Engine/Source/Runtime/TraceLog/Private/Trace/Trace.cpp b/Engine/Source/Runtime/TraceLog/Private/Trace/Trace.cpp
--- a/Engine/Source/Runtime/TraceLog/Private/Trace/Trace.cpp
+++ b/Engine/Source/Runtime/TraceLog/Private/Trace/Trace.cpp
@@ -13,6 +13,7 @@ namespace Private
 
 ////////////////////////////////////////////////////////////////////////////////
 extern UPTRINT	GPendingDataHandle;
+void			Writer_Flush();
 bool			Writer_ConnectImpl(const ANSICHAR*);
 bool			Writer_ToggleEventImpl(const ANSICHAR*, const ANSICHAR*, bool);
 
@@ -54,6 +55,12 @@ bool ToggleEvent(const TCHAR* LoggerName, const TCHAR* EventName, bool State)
 	return Private::Writer_ToggleEventImpl(LoggerNameA, EventNameA, State);
 }
 
+////////////////////////////////////////////////////////////////////////////////
+void Flush()
+{
+	Private::Writer_Flush();
+}
+
 } // namespace Trace
 
 #endif // UE_TRACE_ENABLED
diff --git a/Engine/Source/Runtime/TraceLog/Private/Trace/Writer.cpp b/Engine/Source/Runtime/TraceLog/Private/Trace/Writer.cpp
--- a/Engine/Source/Runtime/TraceLog/Private/Trace/Writer.cpp
+++ b/Engine/Source/Runtime/TraceLog/Private/Trace/Writer.cpp
@@ -203,7 +203,7 @@ static void Writer_UpdateBuffers(void (*DataSink)(const uint8*, uint32))
 }
 
 ////////////////////////////////////////////////////////////////////////////////
-static void Writer_Flush()
+void Writer_Flush()
 {
 	using namespace Private;
 
diff --git a/Engine/Source/Runtime/TraceLog/Public/Trace/Trace.h b/Engine/Source/Runtime/TraceLog/Public/Trace/Trace.h
--- a/Engine/Source/Runtime/TraceLog/Public/Trace/Trace.h
+++ b/Engine/Source/Runtime/TraceLog/Public/Trace/Trace.h
@@ -10,6 +10,7 @@ namespace Trace
 
 UE_TRACE_API bool Connect(const TCHAR* Host) UE_TRACE_IMPL(false);
 UE_TRACE_API bool ToggleEvent(const TCHAR* LoggerName, const TCHAR* EventName, bool State) UE_TRACE_IMPL(false);
+UE_TRACE_API void Flush() UE_TRACE_IMPL();
 
 } // namespace Trace
 
