Removed unused function parameter.

diff --git a/Engine/Source/Runtime/TraceLog/Private/Trace/Writer.cpp b/Engine/Source/Runtime/TraceLog/Private/Trace/Writer.cpp
--- a/Engine/Source/Runtime/TraceLog/Private/Trace/Writer.cpp
+++ b/Engine/Source/Runtime/TraceLog/Private/Trace/Writer.cpp
@@ -145,7 +145,7 @@ static void Writer_ShutdownBuffers()
 }
 
 ////////////////////////////////////////////////////////////////////////////////
-static void Writer_RetireBuffer(FBuffer* Buffer, void (*DataSink)(const uint8*, uint32))
+static void Writer_RetireBuffer(void (*DataSink)(const uint8*, uint32))
 {
 	uint64 StartTsc = Writer_GetTimestamp();
 
@@ -198,7 +198,7 @@ static void Writer_UpdateBuffers(void (*DataSink)(const uint8*, uint32))
 			return;
 		}
 
-		Writer_RetireBuffer(Buffer, DataSink);
+		Writer_RetireBuffer(DataSink);
 	}
 }
 
