Removed Impl() suffix from Connet and ToggleEvent as its redundant.

diff --git a/Engine/Source/Runtime/TraceLog/Private/Trace/Trace.cpp b/Engine/Source/Runtime/TraceLog/Private/Trace/Trace.cpp
--- a/Engine/Source/Runtime/TraceLog/Private/Trace/Trace.cpp
+++ b/Engine/Source/Runtime/TraceLog/Private/Trace/Trace.cpp
@@ -14,8 +14,8 @@ namespace Private
 ////////////////////////////////////////////////////////////////////////////////
 extern UPTRINT	GPendingDataHandle;
 void			Writer_Flush();
-bool			Writer_ConnectImpl(const ANSICHAR*);
-bool			Writer_ToggleEventImpl(const ANSICHAR*, const ANSICHAR*, bool);
+bool			Writer_Connect(const ANSICHAR*);
+bool			Writer_ToggleEvent(const ANSICHAR*, const ANSICHAR*, bool);
 
 } // namespace Private
 
@@ -41,7 +41,7 @@ bool Connect(const TCHAR* InHost)
 	char Host[32];
 	ToAnsiCheap(Host, InHost);
 
-	return Private::Writer_ConnectImpl(Host);
+	return Private::Writer_Connect(Host);
 }
 
 ////////////////////////////////////////////////////////////////////////////////
@@ -52,7 +52,7 @@ bool ToggleEvent(const TCHAR* LoggerName, const TCHAR* EventName, bool State)
 	ToAnsiCheap(LoggerNameA, LoggerName);
 	ToAnsiCheap(EventNameA, EventName);
 
-	return Private::Writer_ToggleEventImpl(LoggerNameA, EventNameA, State);
+	return Private::Writer_ToggleEvent(LoggerNameA, EventNameA, State);
 }
 
 ////////////////////////////////////////////////////////////////////////////////
diff --git a/Engine/Source/Runtime/TraceLog/Private/Trace/Writer.cpp b/Engine/Source/Runtime/TraceLog/Private/Trace/Writer.cpp
--- a/Engine/Source/Runtime/TraceLog/Private/Trace/Writer.cpp
+++ b/Engine/Source/Runtime/TraceLog/Private/Trace/Writer.cpp
@@ -470,8 +470,8 @@ struct FControlCommands
 };
 
 ////////////////////////////////////////////////////////////////////////////////
-bool Writer_ConnectImpl(const ANSICHAR*);
-bool Writer_ToggleEventImpl(const ANSICHAR*, const ANSICHAR*, bool);
+bool Writer_Connect(const ANSICHAR*);
+bool Writer_ToggleEvent(const ANSICHAR*, const ANSICHAR*, bool);
 
 ////////////////////////////////////////////////////////////////////////////////
 static FControlCommands	GControlCommands;
@@ -695,7 +695,7 @@ static void Writer_InitializeControl()
 		{
 			if (ArgC > 0)
 			{
-				Writer_ConnectImpl(ArgV[0]);
+				Writer_Connect(ArgV[0]);
 			}
 		}
 	);
@@ -710,7 +710,7 @@ static void Writer_InitializeControl()
 			const ANSICHAR* LoggerName = ArgV[0];
 			const ANSICHAR* EventName = ArgV[1];
 			const ANSICHAR* State = (ArgC > 2) ? ArgV[2] : "";
-			Writer_ToggleEventImpl(LoggerName, EventName, State[0] != '0');
+			Writer_ToggleEvent(LoggerName, EventName, State[0] != '0');
 		}
 	);
 }
@@ -831,7 +831,7 @@ void Writer_Initialize()
 
 
 ////////////////////////////////////////////////////////////////////////////////
-bool Writer_ConnectImpl(const ANSICHAR* Host)
+bool Writer_Connect(const ANSICHAR* Host)
 {
 	Writer_Initialize();
 
@@ -846,7 +846,7 @@ bool Writer_ConnectImpl(const ANSICHAR* Host)
 }
 
 ////////////////////////////////////////////////////////////////////////////////
-bool Writer_ToggleEventImpl(const ANSICHAR* LoggerName, const ANSICHAR* EventName, bool State)
+bool Writer_ToggleEvent(const ANSICHAR* LoggerName, const ANSICHAR* EventName, bool State)
 {
 	Writer_Initialize();
 
