Implemented ToggleEvent() so that events can be enabled/disabled individually and by logger.

diff --git a/Engine/Source/Runtime/TraceLog/Private/Trace/Event.cpp b/Engine/Source/Runtime/TraceLog/Private/Trace/Event.cpp
--- a/Engine/Source/Runtime/TraceLog/Private/Trace/Event.cpp
+++ b/Engine/Source/Runtime/TraceLog/Private/Trace/Event.cpp
@@ -14,67 +14,12 @@ namespace Trace
 namespace Private
 {
 
-template <typename Type>
-struct TLateAtomic
-{
-	typedef TTraceAtomic<Type> InnerType;
-	InnerType* operator -> () { return (InnerType*)Buffer; }
-	alignas(InnerType) char Buffer[sizeof(InnerType)];
-};
-
-static TLateAtomic<uint32>	GEventUidCounter;		// = 0;
-static TLateAtomic<FEvent*>	GHeadEvent;				// = nullptr;
-void						Writer_Initialize();
+void Writer_EventCreate(FEvent*, const FLiteralName&, const FLiteralName&, const FFieldDesc*, uint32);
 
 } // namespace Private
 
 
 
-////////////////////////////////////////////////////////////////////////////////
-static uint32 GetHashImpl(const ANSICHAR* Input, uint32 Continuation=0x811c9dc5)
-{
-	for (; *Input; ++Input)
-	{
-		Continuation ^= *Input;
-		Continuation *= 0x01000193;
-	}
-	return Continuation;
-}
-
-////////////////////////////////////////////////////////////////////////////////
-static uint32 GetLoggerHash(const ANSICHAR* LoggerName)
-{
-	return GetHashImpl(LoggerName);
-}
-
-////////////////////////////////////////////////////////////////////////////////
-static uint32 GetEventHash(const ANSICHAR* LoggerName, const ANSICHAR* EventName)
-{
-	uint32 Ret = GetHashImpl(LoggerName);
-	Ret = GetHashImpl("@", Ret);
-	Ret = GetHashImpl(EventName, Ret);
-	return Ret;
-}
-
-
-
-////////////////////////////////////////////////////////////////////////////////
-FEvent* FEvent::Find(const ANSICHAR* LoggerName, const ANSICHAR* EventName)
-{
-	uint32 EventHash = GetEventHash(LoggerName, EventName);
-
-	FEvent* Candidate = Private::GHeadEvent->load(std::memory_order_relaxed);
-	for (; Candidate != nullptr; Candidate = (FEvent*)(Candidate->Handle))
-	{
-		if (Candidate->Hash == EventHash)
-		{
-			return Candidate;
-		}
-	}
-
-	return nullptr;
-}
-
 ////////////////////////////////////////////////////////////////////////////////
 void FEvent::Create(
 	FEvent* Target,
@@ -83,86 +28,7 @@ void FEvent::Create(
 	const FFieldDesc* FieldDescs,
 	uint32 FieldCount)
 {
-	using namespace Private;
-
-	Writer_Initialize();
-
-	// Assign a unique ID for this event
-	uint32 Uid = GEventUidCounter->fetch_add(1, std::memory_order_relaxed);
-	Uid += uint32(EKnownEventUids::User);
-
-	if (Uid >= uint32(EKnownEventUids::Max))
-	{
-		Target->bEnabled = false;
-		Target->bInitialized = true;
-		return;
-	}
-
-	// Fill out the target event's properties
-	Target->LoggerHash = GetLoggerHash(LoggerName.Ptr);
-	Target->Hash = GetEventHash(LoggerName.Ptr, EventName.Ptr);
- 	Target->Uid = uint16(Uid);
-	Target->bEnabled = true;
-	Target->bInitialized = true;
-
-	// Calculate the number of fields and size of name data.
-	int NamesSize = LoggerName.Length + EventName.Length;
-	for (uint32 i = 0; i < FieldCount; ++i)
-	{
-		NamesSize += FieldDescs[i].NameSize;
-	}
-
-	// Allocate the new event event in the log stream.
-	uint16 EventSize = sizeof(FNewEventEvent);
-	EventSize += sizeof(FNewEventEvent::Fields[0]) * FieldCount;
-	EventSize += NamesSize;
-	auto& Event = *(FNewEventEvent*)Writer_BeginLog(uint16(EKnownEventUids::NewEvent), EventSize);
-
-	// Write event's main properties.
-	Event.Uid = Target->Uid;
-	Event.LoggerNameSize = LoggerName.Length;
-	Event.EventNameSize = EventName.Length;
-
-	// Write details about event's fields
-	Event.FieldCount = FieldCount;
-	for (uint32 i = 0; i < FieldCount; ++i)
-	{
-		const FFieldDesc& FieldDesc = FieldDescs[i];
-		auto& Out = Event.Fields[i];
-		Out.Offset = FieldDesc.ValueOffset;
-		Out.Size = FieldDesc.ValueSize;
-		Out.TypeInfo = FieldDesc.TypeInfo;
-		Out.NameSize = FieldDesc.NameSize;
-	}
-
-	// Write names
-	uint8* Cursor = (uint8*)(Event.Fields + FieldCount);
-	auto WriteName = [&Cursor] (const ANSICHAR* Data, uint32 Size)
-	{
-		memcpy(Cursor, Data, Size);
-		Cursor += Size;
-	};
-
-	WriteName(LoggerName.Ptr, LoggerName.Length);
-	WriteName(EventName.Ptr, EventName.Length);
-	for (uint32 i = 0; i < FieldCount; ++i)
-	{
-		const FFieldDesc& Desc = FieldDescs[i];
-		WriteName(Desc.Name, Desc.NameSize);
-	}
-
-	Writer_EndLog(&(uint8&)Event);
-
-	// Add this new event into the list so we can look them up later.
-	while (true)
-	{
-		FEvent* HeadEvent = GHeadEvent->load(std::memory_order_relaxed);
-		Target->Handle = HeadEvent;
-		if (GHeadEvent->compare_exchange_weak(HeadEvent, Target, std::memory_order_release))
-		{
-			break;
-		}
-	}
+	Private::Writer_EventCreate(Target, LoggerName, EventName, FieldDescs, FieldCount);
 }
 
 } // namespace Trace
diff --git a/Engine/Source/Runtime/TraceLog/Private/Trace/Trace.cpp b/Engine/Source/Runtime/TraceLog/Private/Trace/Trace.cpp
--- a/Engine/Source/Runtime/TraceLog/Private/Trace/Trace.cpp
+++ b/Engine/Source/Runtime/TraceLog/Private/Trace/Trace.cpp
@@ -15,7 +15,7 @@ namespace Private
 extern UPTRINT	GPendingDataHandle;
 void			Writer_Flush();
 bool			Writer_Connect(const ANSICHAR*);
-bool			Writer_ToggleEvent(const ANSICHAR*, const ANSICHAR*, bool);
+uint32			Writer_EventToggle(const ANSICHAR*, bool);
 
 } // namespace Private
 
@@ -45,14 +45,12 @@ bool Connect(const TCHAR* InHost)
 }
 
 ////////////////////////////////////////////////////////////////////////////////
-bool ToggleEvent(const TCHAR* LoggerName, const TCHAR* EventName, bool State)
+uint32 ToggleEvent(const TCHAR* Wildcard, bool bState)
 {
-	ANSICHAR LoggerNameA[64];
-	ANSICHAR EventNameA[64];
-	ToAnsiCheap(LoggerNameA, LoggerName);
-	ToAnsiCheap(EventNameA, EventName);
+	ANSICHAR WildcardA[64];
+	ToAnsiCheap(WildcardA, Wildcard);
 
-	return Private::Writer_ToggleEvent(LoggerNameA, EventNameA, State);
+	return Private::Writer_EventToggle(WildcardA, bState);
 }
 
 ////////////////////////////////////////////////////////////////////////////////
diff --git a/Engine/Source/Runtime/TraceLog/Private/Trace/Writer.cpp b/Engine/Source/Runtime/TraceLog/Private/Trace/Writer.cpp
--- a/Engine/Source/Runtime/TraceLog/Private/Trace/Writer.cpp
+++ b/Engine/Source/Runtime/TraceLog/Private/Trace/Writer.cpp
@@ -5,6 +5,7 @@
 #if UE_TRACE_ENABLED
 
 #include "Trace/Trace.h"
+#include "Misc/CString.h"
 
 #include <emmintrin.h>
 
@@ -470,8 +471,8 @@ struct FControlCommands
 };
 
 ////////////////////////////////////////////////////////////////////////////////
-bool Writer_Connect(const ANSICHAR*);
-bool Writer_ToggleEvent(const ANSICHAR*, const ANSICHAR*, bool);
+bool	Writer_Connect(const ANSICHAR*);
+uint32	Writer_EventToggle(const ANSICHAR*, bool);
 
 ////////////////////////////////////////////////////////////////////////////////
 static FControlCommands	GControlCommands;
@@ -703,14 +704,13 @@ static void Writer_InitializeControl()
 	Writer_ControlAddCommand("ToggleEvent", nullptr,
 		[] (void*, uint32 ArgC, ANSICHAR const* const* ArgV)
 		{
-			if (ArgC < 2)
+			if (ArgC < 1)
 			{
 				return;
 			}
-			const ANSICHAR* LoggerName = ArgV[0];
-			const ANSICHAR* EventName = ArgV[1];
-			const ANSICHAR* State = (ArgC > 2) ? ArgV[2] : "";
-			Writer_ToggleEvent(LoggerName, EventName, State[0] != '0');
+			const ANSICHAR* Wildcard = ArgV[0];
+			const ANSICHAR* State = (ArgC > 1) ? ArgV[1] : "";
+			Writer_EventToggle(Wildcard, State[0] != '0');
 		}
 	);
 }
@@ -845,18 +845,171 @@ bool Writer_Connect(const ANSICHAR* Host)
 	return true;
 }
 
+
+
+////////////////////////////////////////////////////////////////////////////////
+template <typename Type>
+struct TLateAtomic
+{
+	typedef TTraceAtomic<Type> InnerType;
+	InnerType* operator -> () { return (InnerType*)Buffer; }
+	alignas(InnerType) char Buffer[sizeof(InnerType)];
+};
+
+static TLateAtomic<uint32>	GEventUidCounter;	// = 0;
+static TLateAtomic<FEvent*>	GHeadEvent;			// = nullptr;
+
+////////////////////////////////////////////////////////////////////////////////
+template <typename ElementType>
+static uint32 Writer_EventGetHash(const ElementType* Input, int32 Length=-1)
+{
+	uint32 Result = 0x811c9dc5;
+	for (; *Input && Length; ++Input, --Length)
+	{
+		Result ^= *Input;
+		Result *= 0x01000193;
+	}
+	return Result;
+}
+
+////////////////////////////////////////////////////////////////////////////////
+static uint32 Writer_EventGetHash(uint32 LoggerHash, uint32 NameHash)
+{
+	uint32 Parts[3] = { LoggerHash, NameHash, 0 };
+	return Writer_EventGetHash(Parts);
+}
+
 ////////////////////////////////////////////////////////////////////////////////
-bool Writer_ToggleEvent(const ANSICHAR* LoggerName, const ANSICHAR* EventName, bool State)
+void Writer_EventCreate(
+	FEvent* Target,
+	const FLiteralName& LoggerName,
+	const FLiteralName& EventName,
+	const FFieldDesc* FieldDescs,
+	uint32 FieldCount)
 {
 	Writer_Initialize();
 
-	if (FEvent* Event = FEvent::Find(LoggerName, EventName))
+	// Assign a unique ID for this event
+	uint32 Uid = GEventUidCounter->fetch_add(1, std::memory_order_relaxed);
+	Uid += uint32(EKnownEventUids::User);
+
+	if (Uid >= uint32(EKnownEventUids::Max))
 	{
-		Event->bEnabled = State;
-		return true;
+		Target->bEnabled = false;
+		Target->bInitialized = true;
+		return;
 	}
 
-	return false;
+	uint32 LoggerHash = Writer_EventGetHash(LoggerName.Ptr);
+	uint32 NameHash = Writer_EventGetHash(EventName.Ptr);
+
+	// Fill out the target event's properties
+ 	Target->Uid = uint16(Uid);
+	Target->LoggerHash = LoggerHash;
+	Target->Hash = Writer_EventGetHash(LoggerHash, NameHash);
+	Target->bEnabled = true;
+	Target->bInitialized = true;
+
+	// Calculate the number of fields and size of name data.
+	int NamesSize = LoggerName.Length + EventName.Length;
+	for (uint32 i = 0; i < FieldCount; ++i)
+	{
+		NamesSize += FieldDescs[i].NameSize;
+	}
+
+	// Allocate the new event event in the log stream.
+	uint16 EventSize = sizeof(FNewEventEvent);
+	EventSize += sizeof(FNewEventEvent::Fields[0]) * FieldCount;
+	EventSize += NamesSize;
+	auto& Event = *(FNewEventEvent*)Writer_BeginLog(uint16(EKnownEventUids::NewEvent), EventSize);
+
+	// Write event's main properties.
+	Event.Uid = Target->Uid;
+	Event.LoggerNameSize = LoggerName.Length;
+	Event.EventNameSize = EventName.Length;
+
+	// Write details about event's fields
+	Event.FieldCount = FieldCount;
+	for (uint32 i = 0; i < FieldCount; ++i)
+	{
+		const FFieldDesc& FieldDesc = FieldDescs[i];
+		auto& Out = Event.Fields[i];
+		Out.Offset = FieldDesc.ValueOffset;
+		Out.Size = FieldDesc.ValueSize;
+		Out.TypeInfo = FieldDesc.TypeInfo;
+		Out.NameSize = FieldDesc.NameSize;
+	}
+
+	// Write names
+	uint8* Cursor = (uint8*)(Event.Fields + FieldCount);
+	auto WriteName = [&Cursor] (const ANSICHAR* Data, uint32 Size)
+	{
+		memcpy(Cursor, Data, Size);
+		Cursor += Size;
+	};
+
+	WriteName(LoggerName.Ptr, LoggerName.Length);
+	WriteName(EventName.Ptr, EventName.Length);
+	for (uint32 i = 0; i < FieldCount; ++i)
+	{
+		const FFieldDesc& Desc = FieldDescs[i];
+		WriteName(Desc.Name, Desc.NameSize);
+	}
+
+	Writer_EndLog(&(uint8&)Event);
+
+	// Add this new event into the list so we can look them up later.
+	while (true)
+	{
+		FEvent* HeadEvent = GHeadEvent->load(std::memory_order_relaxed);
+		Target->Handle = HeadEvent;
+		if (GHeadEvent->compare_exchange_weak(HeadEvent, Target, std::memory_order_release))
+		{
+			break;
+		}
+	}
+}
+
+////////////////////////////////////////////////////////////////////////////////
+uint32 Writer_EventToggle(const ANSICHAR* Wildcard, bool bState)
+{
+	Writer_Initialize();
+
+	uint32 ToggleCount = 0;
+
+	const ANSICHAR* Dot = FCStringAnsi::Strchr(Wildcard, '.');
+	if (Dot == nullptr)
+	{
+		uint32 LoggerHash = Writer_EventGetHash(Wildcard);
+
+		FEvent* Event = Private::GHeadEvent->load(std::memory_order_relaxed);
+		for (; Event != nullptr; Event = (FEvent*)(Event->Handle))
+		{
+			if (Event->LoggerHash == LoggerHash)
+			{
+				Event->bEnabled = bState;
+				++ToggleCount;
+			}
+		}
+
+		return ToggleCount;
+	}
+
+	uint32 LoggerHash = Writer_EventGetHash(Wildcard, int(Dot - Wildcard - 1));
+	uint32 NameHash = Writer_EventGetHash(Dot + 1);
+	uint32 EventHash = Writer_EventGetHash(LoggerHash, NameHash);
+
+	FEvent* Event = Private::GHeadEvent->load(std::memory_order_relaxed);
+	for (; Event != nullptr; Event = (FEvent*)(Event->Handle))
+	{
+		if (Event->Hash == EventHash)
+		{
+			Event->bEnabled = bState;
+			++ToggleCount;
+		}
+	}
+
+	return ToggleCount;
 }
 
 } // namespace Private
diff --git a/Engine/Source/Runtime/TraceLog/Public/Trace/Private/Event.h b/Engine/Source/Runtime/TraceLog/Public/Trace/Private/Event.h
--- a/Engine/Source/Runtime/TraceLog/Public/Trace/Private/Event.h
+++ b/Engine/Source/Runtime/TraceLog/Public/Trace/Private/Event.h
@@ -68,7 +68,6 @@ public:
 	bool						bInitialized;
 	bool						bEnabled;
 	UE_TRACE_API static void	Create(FEvent* Target, const FLiteralName& LoggerName, const FLiteralName& EventName, const FFieldDesc* FieldDescs, uint32 FieldCount);
-	static FEvent*				Find(const ANSICHAR* LoggerName, const ANSICHAR* EventName);
 };
 
 } // namespace Trace
diff --git a/Engine/Source/Runtime/TraceLog/Public/Trace/Trace.h b/Engine/Source/Runtime/TraceLog/Public/Trace/Trace.h
--- a/Engine/Source/Runtime/TraceLog/Public/Trace/Trace.h
+++ b/Engine/Source/Runtime/TraceLog/Public/Trace/Trace.h
@@ -8,9 +8,9 @@
 namespace Trace
 {
 
-UE_TRACE_API bool Connect(const TCHAR* Host) UE_TRACE_IMPL(false);
-UE_TRACE_API bool ToggleEvent(const TCHAR* LoggerName, const TCHAR* EventName, bool State) UE_TRACE_IMPL(false);
-UE_TRACE_API void Flush() UE_TRACE_IMPL();
+UE_TRACE_API bool	Connect(const TCHAR* Host) UE_TRACE_IMPL(false);
+UE_TRACE_API uint32 ToggleEvent(const TCHAR* Wildcard, bool bState) UE_TRACE_IMPL(false);
+UE_TRACE_API void	Flush() UE_TRACE_IMPL();
 
 } // namespace Trace
 
