Implemented disabling of events once internal pre-connect storage has been exhausted (which will eventually be change to the condition of a ring-buffer wrapping for the first time).

diff --git a/Engine/Source/Runtime/TraceLog/Private/Trace/Writer.cpp b/Engine/Source/Runtime/TraceLog/Private/Trace/Writer.cpp
--- a/Engine/Source/Runtime/TraceLog/Private/Trace/Writer.cpp
+++ b/Engine/Source/Runtime/TraceLog/Private/Trace/Writer.cpp
@@ -375,6 +375,7 @@ static FHoldBuffer		GHoldBuffer;
 static UPTRINT			GDataHandle			= 0;
 static EDataState		GDataState;			// = EDataState::Passive;
 UPTRINT					GPendingDataHandle	= 0;
+void					Writer_EventDisableUnimportant();
 
 ////////////////////////////////////////////////////////////////////////////////
 static void Writer_UpdateData()
@@ -441,7 +442,7 @@ static void Writer_UpdateData()
 		bool bOverflown = GHoldBuffer.IsFull();
 		if (bOverflown && GDataState != EDataState::Partial)
 		{
-			/* TODO: turn off non-retail events */
+			Writer_EventDisableUnimportant();
 			GDataState = EDataState::Partial;
 		}
 	}
@@ -970,6 +971,16 @@ void Writer_EventCreate(
 	}
 }
 
+////////////////////////////////////////////////////////////////////////////////
+void Writer_EventDisableUnimportant()
+{
+	FEvent* Event = Private::GHeadEvent->load(std::memory_order_relaxed);
+	for (; Event != nullptr; Event = (FEvent*)(Event->Handle))
+	{
+		Event->bEnabled = false;
+	}
+}
+
 ////////////////////////////////////////////////////////////////////////////////
 uint32 Writer_EventToggle(const ANSICHAR* Wildcard, bool bState)
 {
