Used ARRAY_COUNT now that I've found it.

diff --git a/Engine/Source/Runtime/TraceLog/Private/Trace/Writer.cpp b/Engine/Source/Runtime/TraceLog/Private/Trace/Writer.cpp
--- a/Engine/Source/Runtime/TraceLog/Private/Trace/Writer.cpp
+++ b/Engine/Source/Runtime/TraceLog/Private/Trace/Writer.cpp
@@ -597,9 +597,8 @@ static void Writer_ControlRecv()
 			Word,
 		} ParseState = EParseState::CrLfSkip;
 
-		const uint32 ArgN = 16;
 		uint32 ArgC = 0;
-		const ANSICHAR* ArgV[ArgN];
+		const ANSICHAR* ArgV[16];
 
 		const ANSICHAR* __restrict Spent = Buffer;
 		for (ANSICHAR* __restrict Cursor = Buffer; Cursor < Head; ++Cursor)
@@ -620,7 +619,7 @@ static void Writer_ControlRecv()
 					continue;
 				}
 
-				if (ArgC < ArgN)
+				if (ArgC < ARRAY_COUNT(ArgV))
 				{
 					ArgV[ArgC] = Cursor;
 					++ArgC;
