Moved UE_TRACE_* macro details into Private/Trace.h so as to have a cleaner Public/Trace.h

diff --git a/Engine/Source/Runtime/TraceLog/Public/Trace/Private/Trace.h b/Engine/Source/Runtime/TraceLog/Public/Trace/Private/Trace.h
--- a/Engine/Source/Runtime/TraceLog/Public/Trace/Private/Trace.h
+++ b/Engine/Source/Runtime/TraceLog/Public/Trace/Private/Trace.h
@@ -23,3 +23,90 @@
 #	define UE_TRACE_IMPL(...)	{ return __VA_ARGS__; }
 #	define UE_TRACE_API			inline
 #endif
+
+
+
+#if UE_TRACE_ENABLED
+
+#define TRACE_PRIVATE_EVENT_DEFINE(LoggerName, EventName) \
+	Trace::FEvent LoggerName##EventName##Event;
+
+#define TRACE_PRIVATE_EVENT_BEGIN(LoggerName, EventName) \
+	TRACE_PRIVATE_EVENT_BEGIN_IMPL(static, LoggerName, EventName)
+
+#define TRACE_PRIVATE_EVENT_BEGIN_EXTERN(LoggerName, EventName) \
+	TRACE_PRIVATE_EVENT_BEGIN_IMPL(extern, LoggerName, EventName)
+
+#define TRACE_PRIVATE_EVENT_BEGIN_IMPL(LinkageType, LoggerName, EventName) \
+	LinkageType TRACE_PRIVATE_EVENT_DEFINE(LoggerName, EventName) \
+	struct F##LoggerName##EventName##Fields \
+	{ \
+		static void FORCENOINLINE Initialize() \
+		{ \
+			static const bool bOnceOnly = [] () { \
+				F##LoggerName##EventName##Fields Fields; \
+				const auto* Descs = (Trace::FFieldDesc*)(&Fields); \
+				uint32 DescCount = uint32(sizeof(Fields) / sizeof(*Descs)); \
+				const auto& LoggerLiteral = Trace::FLiteralName(#LoggerName); \
+				const auto& EventLiteral = Trace::FLiteralName(#EventName); \
+				Trace::FEvent::Create(&LoggerName##EventName##Event, LoggerLiteral, EventLiteral, Descs, DescCount); \
+				return true; \
+			}(); \
+		} \
+		Trace::TField<0, 
+
+#define TRACE_PRIVATE_EVENT_FIELD(FieldType, FieldName) \
+		FieldType> const FieldName = Trace::FLiteralName(#FieldName); \
+		Trace::TField<decltype(FieldName)::Offset + decltype(FieldName)::Size,
+
+#define TRACE_PRIVATE_EVENT_END() \
+		Trace::EndOfFields> const EventSize_Private = {}; \
+		Trace::TField<decltype(EventSize_Private)::Value, Trace::Attachment> const Attachment = {}; \
+		explicit operator bool () const { return true; } \
+	};
+
+#define TRACE_PRIVATE_EVENT_SIZE(LoggerName, EventName) \
+	decltype(F##LoggerName##EventName##Fields::EventSize_Private)::Value
+
+#define TRACE_PRIVATE_LOG(LoggerName, EventName, ...) \
+	if (!LoggerName##EventName##Event.bInitialized) \
+		F##LoggerName##EventName##Fields::Initialize(); \
+	if (LoggerName##EventName##Event.bEnabled) \
+		if (const auto& __restrict EventName = (F##LoggerName##EventName##Fields&)LoggerName##EventName##Event) \
+			Trace::FEvent::FLogScope(LoggerName##EventName##Event.Uid, TRACE_PRIVATE_EVENT_SIZE(LoggerName, EventName), ##__VA_ARGS__)
+
+#else
+
+#define TRACE_PRIVATE_EVENT_DEFINE(LoggerName, EventName)
+
+#define TRACE_PRIVATE_EVENT_BEGIN(LoggerName, EventName) \
+	TRACE_PRIVATE_EVENT_BEGIN_IMPL(LoggerName, EventName)
+
+#define TRACE_PRIVATE_EVENT_BEGIN_EXTERN(LoggerName, EventName) \
+	TRACE_PRIVATE_EVENT_BEGIN_IMPL(LoggerName, EventName)
+
+#define TRACE_PRIVATE_EVENT_BEGIN_IMPL(LoggerName, EventName) \
+	struct F##LoggerName##EventName##Dummy \
+	{ \
+		struct FTraceDisabled \
+		{ \
+			const FTraceDisabled& operator () (...) const { return *this; } \
+		}; \
+		const F##LoggerName##EventName##Dummy& operator << (const FTraceDisabled&) const \
+		{ \
+			return *this; \
+		} \
+		explicit operator bool () const { return false; }
+
+#define TRACE_PRIVATE_EVENT_FIELD(FieldType, FieldName) \
+		const FTraceDisabled& FieldName;
+
+#define TRACE_PRIVATE_EVENT_END() \
+		const FTraceDisabled& Attachment; \
+	};
+
+#define TRACE_PRIVATE_LOG(LoggerName, EventName, ...) \
+	if (const auto& EventName = *(F##LoggerName##EventName##Dummy*)1) \
+		EventName
+
+#endif // UE_TRACE_ENABLED
diff --git a/Engine/Source/Runtime/TraceLog/Public/Trace/Trace.h b/Engine/Source/Runtime/TraceLog/Public/Trace/Trace.h
--- a/Engine/Source/Runtime/TraceLog/Public/Trace/Trace.h
+++ b/Engine/Source/Runtime/TraceLog/Public/Trace/Trace.h
@@ -14,93 +14,14 @@ UE_TRACE_API void	Flush() UE_TRACE_IMPL();
 
 } // namespace Trace
 
-#if UE_TRACE_ENABLED
-
-#define UE_TRACE_EVENT_DEFINE(LoggerName, EventName) \
-	Trace::FEvent LoggerName##EventName##Event;
-
-#define UE_TRACE_EVENT_BEGIN(LoggerName, EventName) \
-	UE_TRACE_EVENT_BEGIN_IMPL(static, LoggerName, EventName)
-
-#define UE_TRACE_EVENT_BEGIN_EXTERN(LoggerName, EventName) \
-	UE_TRACE_EVENT_BEGIN_IMPL(extern, LoggerName, EventName)
-
-#define UE_TRACE_EVENT_BEGIN_IMPL(LinkageType, LoggerName, EventName) \
-	LinkageType UE_TRACE_EVENT_DEFINE(LoggerName, EventName) \
-	struct F##LoggerName##EventName##Fields \
-	{ \
-		static void FORCENOINLINE Initialize() \
-		{ \
-			static const bool bOnceOnly = [] () { \
-				F##LoggerName##EventName##Fields Fields; \
-				const auto* Descs = (Trace::FFieldDesc*)(&Fields); \
-				uint32 DescCount = uint32(sizeof(Fields) / sizeof(*Descs)); \
-				const auto& LoggerLiteral = Trace::FLiteralName(#LoggerName); \
-				const auto& EventLiteral = Trace::FLiteralName(#EventName); \
-				Trace::FEvent::Create(&LoggerName##EventName##Event, LoggerLiteral, EventLiteral, Descs, DescCount); \
-				return true; \
-			}(); \
-		} \
-		Trace::TField<0, 
-
-#define UE_TRACE_EVENT_FIELD(FieldType, FieldName) \
-		FieldType> const FieldName = Trace::FLiteralName(#FieldName); \
-		Trace::TField<decltype(FieldName)::Offset + decltype(FieldName)::Size,
-
-#define UE_TRACE_EVENT_END() \
-		Trace::EndOfFields> const EventSize_Private = {}; \
-		Trace::TField<decltype(EventSize_Private)::Value, Trace::Attachment> const Attachment = {}; \
-		explicit operator bool () const { return true; } \
-	};
-
-#define UE_TRACE_EVENT_SIZE(LoggerName, EventName) \
-	decltype(F##LoggerName##EventName##Fields::EventSize_Private)::Value
-
-#define UE_TRACE_LOG(LoggerName, EventName, ...) \
-	if (!LoggerName##EventName##Event.bInitialized) \
-		F##LoggerName##EventName##Fields::Initialize(); \
-	if (LoggerName##EventName##Event.bEnabled) \
-		if (const auto& __restrict EventName = (F##LoggerName##EventName##Fields&)LoggerName##EventName##Event) \
-			Trace::FEvent::FLogScope(LoggerName##EventName##Event.Uid, UE_TRACE_EVENT_SIZE(LoggerName, EventName), ##__VA_ARGS__)
-
-#else // UE_TRACE_ENABLED
-
-#define UE_TRACE_EVENT_DEFINE(LoggerName, EventName)
-
-#define UE_TRACE_EVENT_BEGIN(LoggerName, EventName) \
-	UE_TRACE_EVENT_BEGIN_IMPL(LoggerName, EventName)
-
-#define UE_TRACE_EVENT_BEGIN_EXTERN(LoggerName, EventName) \
-	UE_TRACE_EVENT_BEGIN_IMPL(LoggerName, EventName)
-
-#define UE_TRACE_EVENT_BEGIN_IMPL(LoggerName, EventName) \
-	struct F##LoggerName##EventName##Dummy \
-	{ \
-		struct FTraceDisabled \
-		{ \
-			const FTraceDisabled& operator () (...) const { return *this; } \
-		}; \
-		const F##LoggerName##EventName##Dummy& operator << (const FTraceDisabled&) const \
-		{ \
-			return *this; \
-		} \
-		explicit operator bool () const { return false; }
-
-#define UE_TRACE_EVENT_FIELD(FieldType, FieldName) \
-		const FTraceDisabled& FieldName;
-
-#define UE_TRACE_EVENT_END() \
-		const FTraceDisabled& Attachment; \
-	};
-
-#define UE_TRACE_EVENT_SIZE(EventName) \
-	1
-
-#define UE_TRACE_LOG(LoggerName, EventName, ...) \
-	if (const auto& EventName = *(F##LoggerName##EventName##Dummy*)1) \
-		EventName
-
-#endif // UE_TRACE_ENABLED
+////////////////////////////////////////////////////////////////////////////////
+#define UE_TRACE_EVENT_DEFINE(LoggerName, EventName)		TRACE_PRIVATE_EVENT_DEFINE(LoggerName, EventName)
+#define UE_TRACE_EVENT_BEGIN(LoggerName, EventName)			TRACE_PRIVATE_EVENT_BEGIN(LoggerName, EventName)
+#define UE_TRACE_EVENT_BEGIN_EXTERN(LoggerName, EventName)	TRACE_PRIVATE_EVENT_BEGIN_EXTERN(LoggerName, EventName)
+#define UE_TRACE_EVENT_FIELD(FieldType, FieldName)			TRACE_PRIVATE_EVENT_FIELD(FieldType, FieldName)
+#define UE_TRACE_EVENT_END()								TRACE_PRIVATE_EVENT_END()
+#define UE_TRACE_EVENT_SIZE(LoggerName, EventName)			TRACE_PRIVATE_EVENT_SIZE(LoggerName, EventName)
+#define UE_TRACE_LOG(LoggerName, EventName, ...)			TRACE_PRIVATE_LOG(LoggerName, EventName, __VA_ARGS__)
 
 #include "Private/Field.h"
 #include "Private/Event.inl"
