Removed UE_TRACE_EVENT_SIZE() from the public API. This is not something I wish to expose to the user.

diff --git a/Engine/Source/Runtime/TraceLog/Private/Trace/Writer.cpp b/Engine/Source/Runtime/TraceLog/Private/Trace/Writer.cpp
--- a/Engine/Source/Runtime/TraceLog/Private/Trace/Writer.cpp
+++ b/Engine/Source/Runtime/TraceLog/Private/Trace/Writer.cpp
@@ -177,7 +177,7 @@ static void Writer_RetireBuffer(void (*DataSink)(const uint8*, uint32))
 
 	if (true)
 	{
-		uint8 PerfEventBuffer[UE_TRACE_EVENT_SIZE($Trace, PerfWorker) + FEvent::HeaderSize];
+		uint8 PerfEventBuffer[TRACE_PRIVATE_EVENT_SIZE($Trace, PerfWorker) + FEvent::HeaderSize];
 		UE_TRACE_LOG($Trace, PerfWorker, PerfEventBuffer)
 			<< PerfWorker.Start(StartTsc)
 			<< PerfWorker.Acquire(uint32(AcquireTailTsc - StartTsc))
@@ -233,7 +233,7 @@ UE_TRACE_API void* Writer_NextBuffer(Private::FBuffer* Buffer, uint32 PrevUsed,
 {
 	using namespace Private;
 
-	uint32 PerfEventSize = UE_TRACE_EVENT_SIZE($Trace, PerfNextBuffer) + FEvent::HeaderSize;
+	uint32 PerfEventSize = TRACE_PRIVATE_EVENT_SIZE($Trace, PerfNextBuffer) + FEvent::HeaderSize;
 	SizeAndRef += PerfEventSize;
 
 	uint64 StartTsc = Writer_GetTimestamp();
diff --git a/Engine/Source/Runtime/TraceLog/Public/Trace/Trace.h b/Engine/Source/Runtime/TraceLog/Public/Trace/Trace.h
--- a/Engine/Source/Runtime/TraceLog/Public/Trace/Trace.h
+++ b/Engine/Source/Runtime/TraceLog/Public/Trace/Trace.h
@@ -20,7 +20,6 @@ UE_TRACE_API void	Flush() UE_TRACE_IMPL();
 #define UE_TRACE_EVENT_BEGIN_EXTERN(LoggerName, EventName)	TRACE_PRIVATE_EVENT_BEGIN_EXTERN(LoggerName, EventName)
 #define UE_TRACE_EVENT_FIELD(FieldType, FieldName)			TRACE_PRIVATE_EVENT_FIELD(FieldType, FieldName)
 #define UE_TRACE_EVENT_END()								TRACE_PRIVATE_EVENT_END()
-#define UE_TRACE_EVENT_SIZE(LoggerName, EventName)			TRACE_PRIVATE_EVENT_SIZE(LoggerName, EventName)
 #define UE_TRACE_LOG(LoggerName, EventName, ...)			TRACE_PRIVATE_LOG(LoggerName, EventName, __VA_ARGS__)
 
 #include "Private/Field.h"
