Added a UE_TRACE_EVENT_IS_ENABLED() allowing users to check if their event is getting logged.

diff --git a/Engine/Source/Runtime/TraceLog/Public/Trace/Private/Trace.h b/Engine/Source/Runtime/TraceLog/Public/Trace/Private/Trace.h
--- a/Engine/Source/Runtime/TraceLog/Public/Trace/Private/Trace.h
+++ b/Engine/Source/Runtime/TraceLog/Public/Trace/Private/Trace.h
@@ -41,7 +41,7 @@
 	LinkageType TRACE_PRIVATE_EVENT_DEFINE(LoggerName, EventName) \
 	struct F##LoggerName##EventName##Fields \
 	{ \
-		static void FORCENOINLINE Initialize() \
+		static bool FORCENOINLINE Initialize() \
 		{ \
 			static const bool bOnceOnly = [] () { \
 				F##LoggerName##EventName##Fields Fields; \
@@ -52,6 +52,7 @@
 				Trace::FEvent::Create(&LoggerName##EventName##Event, LoggerLiteral, EventLiteral, Descs, DescCount); \
 				return true; \
 			}(); \
+			return true; \
 		} \
 		Trace::TField<0, 
 
@@ -65,13 +66,17 @@
 		explicit operator bool () const { return true; } \
 	};
 
+#define TRACE_PRIVATE_EVENT_IS_ENABLED(LoggerName, EventName) \
+	( \
+		(LoggerName##EventName##Event.bInitialized || F##LoggerName##EventName##Fields::Initialize()) \
+		&& (LoggerName##EventName##Event.bEnabled) \
+	)
+
 #define TRACE_PRIVATE_EVENT_SIZE(LoggerName, EventName) \
 	decltype(F##LoggerName##EventName##Fields::EventSize_Private)::Value
 
 #define TRACE_PRIVATE_LOG(LoggerName, EventName, ...) \
-	if (!LoggerName##EventName##Event.bInitialized) \
-		F##LoggerName##EventName##Fields::Initialize(); \
-	if (LoggerName##EventName##Event.bEnabled) \
+	if (TRACE_PRIVATE_EVENT_IS_ENABLED(LoggerName, EventName)) \
 		if (const auto& __restrict EventName = (F##LoggerName##EventName##Fields&)LoggerName##EventName##Event) \
 			Trace::FEvent::FLogScope(LoggerName##EventName##Event.Uid, TRACE_PRIVATE_EVENT_SIZE(LoggerName, EventName), ##__VA_ARGS__)
 
@@ -105,6 +110,9 @@
 		const FTraceDisabled& Attachment; \
 	};
 
+#define TRACE_PRIVATE_EVENT_IS_ENABLED(LoggerName, EventName) \
+	(false)
+
 #define TRACE_PRIVATE_LOG(LoggerName, EventName, ...) \
 	if (const auto& EventName = *(F##LoggerName##EventName##Dummy*)1) \
 		EventName
diff --git a/Engine/Source/Runtime/TraceLog/Public/Trace/Trace.h b/Engine/Source/Runtime/TraceLog/Public/Trace/Trace.h
--- a/Engine/Source/Runtime/TraceLog/Public/Trace/Trace.h
+++ b/Engine/Source/Runtime/TraceLog/Public/Trace/Trace.h
@@ -20,6 +20,7 @@ UE_TRACE_API void	Flush() UE_TRACE_IMPL();
 #define UE_TRACE_EVENT_BEGIN_EXTERN(LoggerName, EventName)	TRACE_PRIVATE_EVENT_BEGIN_EXTERN(LoggerName, EventName)
 #define UE_TRACE_EVENT_FIELD(FieldType, FieldName)			TRACE_PRIVATE_EVENT_FIELD(FieldType, FieldName)
 #define UE_TRACE_EVENT_END()								TRACE_PRIVATE_EVENT_END()
+#define UE_TRACE_EVENT_IS_ENABLED(LoggerName, EventName)	TRACE_PRIVATE_EVENT_IS_ENABLED(LoggerName, EventName)
 #define UE_TRACE_LOG(LoggerName, EventName, ...)			TRACE_PRIVATE_LOG(LoggerName, EventName, __VA_ARGS__)
 
 #include "Private/Field.h"
