Fixed an angry Clang compiler.

diff --git a/Engine/Source/Runtime/TraceLog/Public/Trace/Trace.h b/Engine/Source/Runtime/TraceLog/Public/Trace/Trace.h
--- a/Engine/Source/Runtime/TraceLog/Public/Trace/Trace.h
+++ b/Engine/Source/Runtime/TraceLog/Public/Trace/Trace.h
@@ -21,7 +21,7 @@ UE_TRACE_API void	Flush() UE_TRACE_IMPL();
 #define UE_TRACE_EVENT_FIELD(FieldType, FieldName)			TRACE_PRIVATE_EVENT_FIELD(FieldType, FieldName)
 #define UE_TRACE_EVENT_END()								TRACE_PRIVATE_EVENT_END()
 #define UE_TRACE_EVENT_IS_ENABLED(LoggerName, EventName)	TRACE_PRIVATE_EVENT_IS_ENABLED(LoggerName, EventName)
-#define UE_TRACE_LOG(LoggerName, EventName, ...)			TRACE_PRIVATE_LOG(LoggerName, EventName, __VA_ARGS__)
+#define UE_TRACE_LOG(LoggerName, EventName, ...)			TRACE_PRIVATE_LOG(LoggerName, EventName, ##__VA_ARGS__)
 
 #include "Private/Field.h"
 #include "Private/Event.inl"
