Only attempt to send Trace's performance events if they are enabled

diff --git a/Engine/Source/Runtime/TraceLog/Private/Trace/Writer.cpp b/Engine/Source/Runtime/TraceLog/Private/Trace/Writer.cpp
--- a/Engine/Source/Runtime/TraceLog/Private/Trace/Writer.cpp
+++ b/Engine/Source/Runtime/TraceLog/Private/Trace/Writer.cpp
@@ -175,7 +175,7 @@ static void Writer_RetireBuffer(void (*DataSink)(const uint8*, uint32))
 	GTailBuffer = Next;
 	uint64 DoneTsc = Writer_GetTimestamp();
 
-	if (true)
+	if (UE_TRACE_EVENT_IS_ENABLED($Trace, PerfWorker))
 	{
 		uint8 PerfEventBuffer[TRACE_PRIVATE_EVENT_SIZE($Trace, PerfWorker) + FEvent::HeaderSize];
 		UE_TRACE_LOG($Trace, PerfWorker, PerfEventBuffer)
@@ -233,7 +233,11 @@ UE_TRACE_API void* Writer_NextBuffer(Private::FBuffer* Buffer, uint32 PrevUsed,
 {
 	using namespace Private;
 
-	uint32 PerfEventSize = TRACE_PRIVATE_EVENT_SIZE($Trace, PerfNextBuffer) + FEvent::HeaderSize;
+	uint32 PerfEventSize = 0;
+	if (UE_TRACE_EVENT_IS_ENABLED($Trace, PerfNextBuffer))
+	{
+		PerfEventSize = TRACE_PRIVATE_EVENT_SIZE($Trace, PerfNextBuffer) + FEvent::HeaderSize;
+	}
 	SizeAndRef += PerfEventSize;
 
 	uint64 StartTsc = Writer_GetTimestamp();
@@ -280,12 +284,15 @@ UE_TRACE_API void* Writer_NextBuffer(Private::FBuffer* Buffer, uint32 PrevUsed,
 	PrevUsed &= BufferSizeMask;
 	uint8* Out = (uint8*)(UPTRINT(NextBuffer) + PrevUsed);
 
-	uint64 DoneTsc = Writer_GetTimestamp();
-	UE_TRACE_LOG($Trace, PerfNextBuffer, Out)
-		<< PerfNextBuffer.Start(StartTsc)
-		<< PerfNextBuffer.Acquire(uint32(AcquireTsc - StartTsc))
-		<< PerfNextBuffer.Done(uint32(DoneTsc - StartTsc))
-		<< PerfNextBuffer.ThreadId(uint16(ThreadGetCurrentId()));
+	if (PerfEventSize)
+	{
+		uint64 DoneTsc = Writer_GetTimestamp();
+		UE_TRACE_LOG($Trace, PerfNextBuffer, Out)
+			<< PerfNextBuffer.Start(StartTsc)
+			<< PerfNextBuffer.Acquire(uint32(AcquireTsc - StartTsc))
+			<< PerfNextBuffer.Done(uint32(DoneTsc - StartTsc))
+			<< PerfNextBuffer.ThreadId(uint16(ThreadGetCurrentId()));
+	}
 
 	return Out + PerfEventSize;
 }
