Removed the auto-disabling of unimportant events. They will be off by default in the not-to-distant future.

diff --git a/Engine/Source/Runtime/TraceLog/Private/Trace/Writer.cpp b/Engine/Source/Runtime/TraceLog/Private/Trace/Writer.cpp
--- a/Engine/Source/Runtime/TraceLog/Private/Trace/Writer.cpp
+++ b/Engine/Source/Runtime/TraceLog/Private/Trace/Writer.cpp
@@ -382,7 +382,6 @@ static FHoldBuffer		GHoldBuffer;
 static UPTRINT			GDataHandle			= 0;
 static EDataState		GDataState;			// = EDataState::Passive;
 UPTRINT					GPendingDataHandle	= 0;
-void					Writer_EventDisableUnimportant();
 
 ////////////////////////////////////////////////////////////////////////////////
 static void Writer_UpdateData()
@@ -449,7 +448,6 @@ static void Writer_UpdateData()
 		bool bOverflown = GHoldBuffer.IsFull();
 		if (bOverflown && GDataState != EDataState::Partial)
 		{
-			Writer_EventDisableUnimportant();
 			GDataState = EDataState::Partial;
 		}
 	}
@@ -977,16 +975,6 @@ void Writer_EventCreate(
 	}
 }
 
-////////////////////////////////////////////////////////////////////////////////
-void Writer_EventDisableUnimportant()
-{
-	FEvent* Event = Private::GHeadEvent->load(std::memory_order_relaxed);
-	for (; Event != nullptr; Event = (FEvent*)(Event->Handle))
-	{
-		Event->bEnabled = false;
-	}
-}
-
 ////////////////////////////////////////////////////////////////////////////////
 uint32 Writer_EventToggle(const ANSICHAR* Wildcard, bool bState)
 {
