Improved handling of a failed handshake when connecting to a recorder

diff --git a/Engine/Source/Runtime/TraceLog/Private/Trace/Writer.cpp b/Engine/Source/Runtime/TraceLog/Private/Trace/Writer.cpp
--- a/Engine/Source/Runtime/TraceLog/Private/Trace/Writer.cpp
+++ b/Engine/Source/Runtime/TraceLog/Private/Trace/Writer.cpp
@@ -397,22 +397,31 @@ static void Writer_UpdateData()
 		{
 			GDataHandle = GPendingDataHandle;
 			GPendingDataHandle = 0;
-			GDataState = EDataState::Sending;
 
 			// Handshake.
 			const uint32 Magic = 'TRCE';
-			TcpSocketSend(GDataHandle, &Magic, sizeof(Magic));
+			bool bOk = TcpSocketSend(GDataHandle, &Magic, sizeof(Magic));
 
 			// Stream header
 			const struct {
 				uint8 Format;
 				uint8 Parameter;
 			} TransportHeader = { 1 };
-			TcpSocketSend(GDataHandle, &TransportHeader, sizeof(TransportHeader));
+			bOk &= TcpSocketSend(GDataHandle, &TransportHeader, sizeof(TransportHeader));
 
 			// Passively collected data
-			TcpSocketSend(GDataHandle, GHoldBuffer.GetData(), GHoldBuffer.GetSize());
-			GHoldBuffer.Shutdown();
+			bOk &= TcpSocketSend(GDataHandle, GHoldBuffer.GetData(), GHoldBuffer.GetSize());
+
+			if (bOk)
+			{
+				GDataState = EDataState::Sending;
+				GHoldBuffer.Shutdown();
+			}
+			else
+			{
+				TcpSocketClose(GDataHandle);
+				GDataHandle = 0;
+			}
 		}
 	}
 
