Events are now off by default until explicitly opted-in to or events that are flagged as "Always". Always events are always logged regardless of opt-in status.

diff --git a/Engine/Source/Runtime/TraceLog/Private/Trace/Event.cpp b/Engine/Source/Runtime/TraceLog/Private/Trace/Event.cpp
--- a/Engine/Source/Runtime/TraceLog/Private/Trace/Event.cpp
+++ b/Engine/Source/Runtime/TraceLog/Private/Trace/Event.cpp
@@ -14,7 +14,7 @@ namespace Trace
 namespace Private
 {
 
-void Writer_EventCreate(FEvent*, const FLiteralName&, const FLiteralName&, const FFieldDesc*, uint32);
+void Writer_EventCreate(FEvent*, const FLiteralName&, const FLiteralName&, const FFieldDesc*, uint32, uint32);
 
 } // namespace Private
 
@@ -26,9 +26,10 @@ void FEvent::Create(
 	const FLiteralName& LoggerName,
 	const FLiteralName& EventName,
 	const FFieldDesc* FieldDescs,
-	uint32 FieldCount)
+	uint32 FieldCount,
+	uint32 Flags)
 {
-	Private::Writer_EventCreate(Target, LoggerName, EventName, FieldDescs, FieldCount);
+	Private::Writer_EventCreate(Target, LoggerName, EventName, FieldDescs, FieldCount, Flags);
 }
 
 } // namespace Trace
diff --git a/Engine/Source/Runtime/TraceLog/Private/Trace/Writer.cpp b/Engine/Source/Runtime/TraceLog/Private/Trace/Writer.cpp
--- a/Engine/Source/Runtime/TraceLog/Private/Trace/Writer.cpp
+++ b/Engine/Source/Runtime/TraceLog/Private/Trace/Writer.cpp
@@ -89,7 +89,7 @@ void Writer_InitializeTiming()
 {
 	GStartCycle = TimeGetTimestamp();
 
-	UE_TRACE_EVENT_BEGIN($Trace, Timing)
+	UE_TRACE_EVENT_BEGIN($Trace, Timing, Always)
 		UE_TRACE_EVENT_FIELD(uint64, StartCycle)
 		UE_TRACE_EVENT_FIELD(uint64, CycleFrequency)
 	UE_TRACE_EVENT_END()
@@ -766,7 +766,7 @@ static bool GInitialized = false;
 ////////////////////////////////////////////////////////////////////////////////
 static void Writer_LogHeader()
 {
-	UE_TRACE_EVENT_BEGIN($Trace, NewTrace)
+	UE_TRACE_EVENT_BEGIN($Trace, NewTrace, Always)
 		UE_TRACE_EVENT_FIELD(uint16, Endian)
 		UE_TRACE_EVENT_FIELD(uint8, Version)
 		UE_TRACE_EVENT_FIELD(uint8, PointerSize)
@@ -899,7 +899,8 @@ void Writer_EventCreate(
 	const FLiteralName& LoggerName,
 	const FLiteralName& EventName,
 	const FFieldDesc* FieldDescs,
-	uint32 FieldCount)
+	uint32 FieldCount,
+	uint32 Flags)
 {
 	Writer_Initialize();
 
@@ -909,7 +910,8 @@ void Writer_EventCreate(
 
 	if (Uid >= uint32(EKnownEventUids::Max))
 	{
-		Target->bEnabled = false;
+		Target->Enabled.bOptedIn = false;
+		Target->Enabled.Internal = 0;
 		Target->bInitialized = true;
 		return;
 	}
@@ -921,7 +923,8 @@ void Writer_EventCreate(
  	Target->Uid = uint16(Uid);
 	Target->LoggerHash = LoggerHash;
 	Target->Hash = Writer_EventGetHash(LoggerHash, NameHash);
-	Target->bEnabled = true;
+	Target->Enabled.Internal = !!(Flags & FEvent::Flag_Always);
+	Target->Enabled.bOptedIn = false;
 	Target->bInitialized = true;
 
 	// Calculate the number of fields and size of name data.
@@ -1001,7 +1004,7 @@ uint32 Writer_EventToggle(const ANSICHAR* Wildcard, bool bState)
 		{
 			if (Event->LoggerHash == LoggerHash)
 			{
-				Event->bEnabled = bState;
+				Event->Enabled.bOptedIn = bState;
 				++ToggleCount;
 			}
 		}
@@ -1018,7 +1021,7 @@ uint32 Writer_EventToggle(const ANSICHAR* Wildcard, bool bState)
 	{
 		if (Event->Hash == EventHash)
 		{
-			Event->bEnabled = bState;
+			Event->Enabled.bOptedIn = bState;
 			++ToggleCount;
 		}
 	}
diff --git a/Engine/Source/Runtime/TraceLog/Public/Trace/Private/Event.h b/Engine/Source/Runtime/TraceLog/Public/Trace/Private/Event.h
--- a/Engine/Source/Runtime/TraceLog/Public/Trace/Private/Event.h
+++ b/Engine/Source/Runtime/TraceLog/Public/Trace/Private/Event.h
@@ -50,6 +50,11 @@ class FEvent
 public:
 	enum { HeaderSize = sizeof(uint32), }; // I'll just dump this here for now.
 
+	enum
+	{
+		Flag_Always	= 1 << 0,
+	};
+
 	class FLogScope
 	{
 	public:
@@ -65,9 +70,17 @@ public:
 	uint32						LoggerHash;
 	uint32						Hash;
 	uint16						Uid;
+	union
+	{
+		struct
+		{
+			bool				bOptedIn;
+			uint8				Internal;
+		};
+		uint16					Test;
+	}							Enabled;
 	bool						bInitialized;
-	bool						bEnabled;
-	UE_TRACE_API static void	Create(FEvent* Target, const FLiteralName& LoggerName, const FLiteralName& EventName, const FFieldDesc* FieldDescs, uint32 FieldCount);
+	UE_TRACE_API static void	Create(FEvent* Target, const FLiteralName& LoggerName, const FLiteralName& EventName, const FFieldDesc* FieldDescs, uint32 FieldCount, uint32 Flags=0);
 };
 
 } // namespace Trace
diff --git a/Engine/Source/Runtime/TraceLog/Public/Trace/Private/Trace.h b/Engine/Source/Runtime/TraceLog/Public/Trace/Private/Trace.h
--- a/Engine/Source/Runtime/TraceLog/Public/Trace/Private/Trace.h
+++ b/Engine/Source/Runtime/TraceLog/Public/Trace/Private/Trace.h
@@ -31,25 +31,27 @@
 #define TRACE_PRIVATE_EVENT_DEFINE(LoggerName, EventName) \
 	Trace::FEvent LoggerName##EventName##Event;
 
-#define TRACE_PRIVATE_EVENT_BEGIN(LoggerName, EventName) \
-	TRACE_PRIVATE_EVENT_BEGIN_IMPL(static, LoggerName, EventName)
+#define TRACE_PRIVATE_EVENT_BEGIN(LoggerName, EventName, ...) \
+	TRACE_PRIVATE_EVENT_BEGIN_IMPL(static, LoggerName, EventName, ##__VA_ARGS__)
 
-#define TRACE_PRIVATE_EVENT_BEGIN_EXTERN(LoggerName, EventName) \
-	TRACE_PRIVATE_EVENT_BEGIN_IMPL(extern, LoggerName, EventName)
+#define TRACE_PRIVATE_EVENT_BEGIN_EXTERN(LoggerName, EventName, ...) \
+	TRACE_PRIVATE_EVENT_BEGIN_IMPL(extern, LoggerName, EventName, ##__VA_ARGS__)
 
-#define TRACE_PRIVATE_EVENT_BEGIN_IMPL(LinkageType, LoggerName, EventName) \
+#define TRACE_PRIVATE_EVENT_BEGIN_IMPL(LinkageType, LoggerName, EventName, ...) \
 	LinkageType TRACE_PRIVATE_EVENT_DEFINE(LoggerName, EventName) \
 	struct F##LoggerName##EventName##Fields \
 	{ \
 		static bool FORCENOINLINE Initialize() \
 		{ \
 			static const bool bOnceOnly = [] () { \
+				const uint32 Always = Trace::FEvent::Flag_Always; \
+				const uint32 Flags = (0, ##__VA_ARGS__); \
 				F##LoggerName##EventName##Fields Fields; \
 				const auto* Descs = (Trace::FFieldDesc*)(&Fields); \
 				uint32 DescCount = uint32(sizeof(Fields) / sizeof(*Descs)); \
 				const auto& LoggerLiteral = Trace::FLiteralName(#LoggerName); \
 				const auto& EventLiteral = Trace::FLiteralName(#EventName); \
-				Trace::FEvent::Create(&LoggerName##EventName##Event, LoggerLiteral, EventLiteral, Descs, DescCount); \
+				Trace::FEvent::Create(&LoggerName##EventName##Event, LoggerLiteral, EventLiteral, Descs, DescCount, Flags); \
 				return true; \
 			}(); \
 			return true; \
@@ -69,7 +71,7 @@
 #define TRACE_PRIVATE_EVENT_IS_ENABLED(LoggerName, EventName) \
 	( \
 		(LoggerName##EventName##Event.bInitialized || F##LoggerName##EventName##Fields::Initialize()) \
-		&& (LoggerName##EventName##Event.bEnabled) \
+		&& (LoggerName##EventName##Event.Enabled.Test) \
 	)
 
 #define TRACE_PRIVATE_EVENT_SIZE(LoggerName, EventName) \
diff --git a/Engine/Source/Runtime/TraceLog/Public/Trace/Trace.h b/Engine/Source/Runtime/TraceLog/Public/Trace/Trace.h
--- a/Engine/Source/Runtime/TraceLog/Public/Trace/Trace.h
+++ b/Engine/Source/Runtime/TraceLog/Public/Trace/Trace.h
@@ -15,13 +15,13 @@ UE_TRACE_API void	Flush() UE_TRACE_IMPL();
 } // namespace Trace
 
 ////////////////////////////////////////////////////////////////////////////////
-#define UE_TRACE_EVENT_DEFINE(LoggerName, EventName)		TRACE_PRIVATE_EVENT_DEFINE(LoggerName, EventName)
-#define UE_TRACE_EVENT_BEGIN(LoggerName, EventName)			TRACE_PRIVATE_EVENT_BEGIN(LoggerName, EventName)
-#define UE_TRACE_EVENT_BEGIN_EXTERN(LoggerName, EventName)	TRACE_PRIVATE_EVENT_BEGIN_EXTERN(LoggerName, EventName)
-#define UE_TRACE_EVENT_FIELD(FieldType, FieldName)			TRACE_PRIVATE_EVENT_FIELD(FieldType, FieldName)
-#define UE_TRACE_EVENT_END()								TRACE_PRIVATE_EVENT_END()
-#define UE_TRACE_EVENT_IS_ENABLED(LoggerName, EventName)	TRACE_PRIVATE_EVENT_IS_ENABLED(LoggerName, EventName)
-#define UE_TRACE_LOG(LoggerName, EventName, ...)			TRACE_PRIVATE_LOG(LoggerName, EventName, ##__VA_ARGS__)
+#define UE_TRACE_EVENT_DEFINE(LoggerName, EventName)			TRACE_PRIVATE_EVENT_DEFINE(LoggerName, EventName)
+#define UE_TRACE_EVENT_BEGIN(LoggerName, EventName, ...)		TRACE_PRIVATE_EVENT_BEGIN(LoggerName, EventName, ##__VA_ARGS__)
+#define UE_TRACE_EVENT_BEGIN_EXTERN(LoggerName, EventName, ...)	TRACE_PRIVATE_EVENT_BEGIN_EXTERN(LoggerName, EventName, ##__VA_ARGS__)
+#define UE_TRACE_EVENT_FIELD(FieldType, FieldName)				TRACE_PRIVATE_EVENT_FIELD(FieldType, FieldName)
+#define UE_TRACE_EVENT_END()									TRACE_PRIVATE_EVENT_END()
+#define UE_TRACE_EVENT_IS_ENABLED(LoggerName, EventName)		TRACE_PRIVATE_EVENT_IS_ENABLED(LoggerName, EventName)
+#define UE_TRACE_LOG(LoggerName, EventName, ...)				TRACE_PRIVATE_LOG(LoggerName, EventName, ##__VA_ARGS__)
 
 #include "Private/Field.h"
 #include "Private/Event.inl"
