Disabled runtime tracing in shipping-variant builds

diff --git a/Engine/Source/Runtime/TraceLog/Public/Trace/Private/Trace.h b/Engine/Source/Runtime/TraceLog/Public/Trace/Private/Trace.h
--- a/Engine/Source/Runtime/TraceLog/Public/Trace/Private/Trace.h
+++ b/Engine/Source/Runtime/TraceLog/Public/Trace/Private/Trace.h
@@ -5,10 +5,11 @@
 #include "CoreTypes.h"
 
 ////////////////////////////////////////////////////////////////////////////////
-#if !IS_PROGRAM && !defined(UE_TRACE_ENABLED)
-#	if PLATFORM_WINDOWS || PLATFORM_PS4 || PLATFORM_XBOXONE
-#		define UE_TRACE_ENABLED	1
-#	else
+#if !defined(UE_TRACE_ENABLED)
+#	if !UE_BUILD_SHIPPING && !IS_PROGRAM
+#		if PLATFORM_WINDOWS || PLATFORM_PS4 || PLATFORM_XBOXONE
+#			define UE_TRACE_ENABLED	1
+#		endif
 #	endif
 #endif // !IS_PROGRAM
 
