Fixed rule break

diff --git a/Engine/Source/Runtime/TraceLog/Public/Trace/Private/Trace.h b/Engine/Source/Runtime/TraceLog/Public/Trace/Private/Trace.h
--- a/Engine/Source/Runtime/TraceLog/Public/Trace/Private/Trace.h
+++ b/Engine/Source/Runtime/TraceLog/Public/Trace/Private/Trace.h
@@ -44,7 +44,8 @@
 	{ \
 		static bool FORCENOINLINE Initialize() \
 		{ \
-			static const bool bOnceOnly = [] () { \
+			static const bool bOnceOnly = [] () \
+			{ \
 				const uint32 Always = Trace::FEvent::Flag_Always; \
 				const uint32 Flags = (0, ##__VA_ARGS__); \
 				F##LoggerName##EventName##Fields Fields; \
