Be explicit about the control flow around event initialisation instead of using the return value from FEvent::Init()

diff --git a/Engine/Source/Runtime/TraceLog/Public/Trace/Private/Trace.h b/Engine/Source/Runtime/TraceLog/Public/Trace/Private/Trace.h
--- a/Engine/Source/Runtime/TraceLog/Public/Trace/Private/Trace.h
+++ b/Engine/Source/Runtime/TraceLog/Public/Trace/Private/Trace.h
@@ -42,7 +42,7 @@
 	LinkageType TRACE_PRIVATE_EVENT_DEFINE(LoggerName, EventName) \
 	struct F##LoggerName##EventName##Fields \
 	{ \
-		static bool FORCENOINLINE Initialize() \
+		static void FORCENOINLINE Initialize() \
 		{ \
 			static const bool bOnceOnly = [] () \
 			{ \
@@ -56,7 +56,6 @@
 				Trace::FEvent::Create(&LoggerName##EventName##Event, LoggerLiteral, EventLiteral, Descs, DescCount, Flags); \
 				return true; \
 			}(); \
-			return true; \
 		} \
 		Trace::TField<0, 
 
@@ -72,7 +71,7 @@
 
 #define TRACE_PRIVATE_EVENT_IS_ENABLED(LoggerName, EventName) \
 	( \
-		(LoggerName##EventName##Event.bInitialized || F##LoggerName##EventName##Fields::Initialize()) \
+		(LoggerName##EventName##Event.bInitialized || (F##LoggerName##EventName##Fields::Initialize(), true)) \
 		&& (LoggerName##EventName##Event.Enabled.Test) \
 	)
 
