Moved EKnownEventUids from public to private

diff --git a/Engine/Source/Developer/TraceAnalysis/Private/Analysis/Engine.cpp b/Engine/Source/Developer/TraceAnalysis/Private/Analysis/Engine.cpp
--- a/Engine/Source/Developer/TraceAnalysis/Private/Analysis/Engine.cpp
+++ b/Engine/Source/Developer/TraceAnalysis/Private/Analysis/Engine.cpp
@@ -300,8 +300,9 @@ FAnalysisEngine::FAnalysisEngine(TArray<IAnalyzer*>&& InAnalyzers)
 	uint16 SelfIndex = Analyzers.Num();
 	Analyzers.Add(this);
 
-	// Manually event routing for known events, and those we don't quite know yet
-	FDispatch& NewEventDispatch = AddDispatch(uint16(EKnownEventUids::NewEvent), 0);
+	// Manually add event routing for known events, and those we don't quite know
+	// yet but are expecting.
+	FDispatch& NewEventDispatch = AddDispatch(uint16(FNewEventEvent::Uid), 0);
 	NewEventDispatch.FirstRoute = 0;
 	AddRoute(SelfIndex, RouteId_NewEvent, 0);
 	AddRoute(SelfIndex, RouteId_NewTrace, "$Trace", "NewTrace");
diff --git a/Engine/Source/Runtime/TraceLog/Private/Trace/Writer.cpp b/Engine/Source/Runtime/TraceLog/Private/Trace/Writer.cpp
--- a/Engine/Source/Runtime/TraceLog/Private/Trace/Writer.cpp
+++ b/Engine/Source/Runtime/TraceLog/Private/Trace/Writer.cpp
@@ -873,6 +873,14 @@ struct TLateAtomic
 static TLateAtomic<uint32>	GEventUidCounter;	// = 0;
 static TLateAtomic<FEvent*>	GHeadEvent;			// = nullptr;
 
+////////////////////////////////////////////////////////////////////////////////
+enum class EKnownEventUids : uint16
+{
+	NewEvent		= FNewEventEvent::Uid,
+	User,
+	Max				= 1 << 14, // ...leaves two MSB bits for other uses.
+};
+
 ////////////////////////////////////////////////////////////////////////////////
 template <typename ElementType>
 static uint32 Writer_EventGetHash(const ElementType* Input, int32 Length=-1)
@@ -941,7 +949,7 @@ void Writer_EventCreate(
 	auto& Event = *(FNewEventEvent*)Writer_BeginLog(uint16(EKnownEventUids::NewEvent), EventSize);
 
 	// Write event's main properties.
-	Event.Uid = Target->Uid;
+	Event.EventUid = Target->Uid;
 	Event.LoggerNameSize = LoggerName.Length;
 	Event.EventNameSize = EventName.Length;
 
diff --git a/Engine/Source/Runtime/TraceLog/Public/Trace/Private/Event.h b/Engine/Source/Runtime/TraceLog/Public/Trace/Private/Event.h
--- a/Engine/Source/Runtime/TraceLog/Public/Trace/Private/Event.h
+++ b/Engine/Source/Runtime/TraceLog/Public/Trace/Private/Event.h
@@ -7,18 +7,12 @@
 namespace Trace
 {
 
-////////////////////////////////////////////////////////////////////////////////
-enum class EKnownEventUids : uint16
-{
-	NewEvent,
-	User,
-	Max			= 1 << 14, // ...leaves two MSB bits for other uses.
-};
-
 ////////////////////////////////////////////////////////////////////////////////
 struct FNewEventEvent
 {
-	uint16		Uid;
+	enum : uint16 { Uid = 0 };
+
+	uint16		EventUid;
 	uint16		FieldCount;
 	uint8		LoggerNameSize;
 	uint8		EventNameSize;
