Strategic change to clearly invalidate events that could not be initialised

diff --git a/Engine/Source/Runtime/TraceLog/Private/Trace/Writer.cpp b/Engine/Source/Runtime/TraceLog/Private/Trace/Writer.cpp
--- a/Engine/Source/Runtime/TraceLog/Private/Trace/Writer.cpp
+++ b/Engine/Source/Runtime/TraceLog/Private/Trace/Writer.cpp
@@ -878,7 +878,8 @@ enum class EKnownEventUids : uint16
 {
 	NewEvent		= FNewEventEvent::Uid,
 	User,
-	Max				= 1 << 14, // ...leaves two MSB bits for other uses.
+	Max				= (1 << 14) - 1, // ...leaves two MSB bits for other uses.
+	Invalid			= Max,
 };
 
 ////////////////////////////////////////////////////////////////////////////////
@@ -918,6 +919,7 @@ void Writer_EventCreate(
 
 	if (Uid >= uint32(EKnownEventUids::Max))
 	{
+		Target->Uid = uint16(EKnownEventUids::Invalid);
 		Target->Enabled.bOptedIn = false;
 		Target->Enabled.Internal = 0;
 		Target->bInitialized = true;
