Fixed Linux compile error.

diff --git a/Engine/Source/Runtime/TraceLog/Public/Trace/Private/Trace.h b/Engine/Source/Runtime/TraceLog/Public/Trace/Private/Trace.h
--- a/Engine/Source/Runtime/TraceLog/Public/Trace/Private/Trace.h
+++ b/Engine/Source/Runtime/TraceLog/Public/Trace/Private/Trace.h
@@ -87,7 +87,7 @@
 
 #define TRACE_PRIVATE_EVENT_DEFINE(LoggerName, EventName)
 
-#define TRACE_PRIVATE_EVENT_BEGIN(LoggerName, EventName) \
+#define TRACE_PRIVATE_EVENT_BEGIN(LoggerName, EventName, ...) \
 	TRACE_PRIVATE_EVENT_BEGIN_IMPL(LoggerName, EventName)
 
 #define TRACE_PRIVATE_EVENT_BEGIN_EXTERN(LoggerName, EventName) \
