Fixed off by one error in Writer_ToggleEvent

diff --git a/Engine/Source/Runtime/TraceLog/Private/Trace/Writer.cpp b/Engine/Source/Runtime/TraceLog/Private/Trace/Writer.cpp
--- a/Engine/Source/Runtime/TraceLog/Private/Trace/Writer.cpp
+++ b/Engine/Source/Runtime/TraceLog/Private/Trace/Writer.cpp
@@ -1022,7 +1022,7 @@ uint32 Writer_EventToggle(const ANSICHAR* Wildcard, bool bState)
 		return ToggleCount;
 	}
 
-	uint32 LoggerHash = Writer_EventGetHash(Wildcard, int(Dot - Wildcard - 1));
+	uint32 LoggerHash = Writer_EventGetHash(Wildcard, int(Dot - Wildcard));
 	uint32 NameHash = Writer_EventGetHash(Dot + 1);
 	uint32 EventHash = Writer_EventGetHash(LoggerHash, NameHash);
 
