Reduced the size of the hold buffer by rather a lot.

diff --git a/Engine/Source/Runtime/TraceLog/Private/Trace/Writer.cpp b/Engine/Source/Runtime/TraceLog/Private/Trace/Writer.cpp
--- a/Engine/Source/Runtime/TraceLog/Private/Trace/Writer.cpp
+++ b/Engine/Source/Runtime/TraceLog/Private/Trace/Writer.cpp
@@ -313,7 +313,7 @@ public:
 private:
 	static const uint32	PageShift = 16;
 	static const uint32	PageSize = 1 << PageShift;
-	static const uint32	MaxPages = (4 * 1024 * 1024) >> PageShift;
+	static const uint32	MaxPages = (64 * 1024) >> PageShift;
 	uint8*				Base;
 	int32				Used;
 	uint16				MappedPageCount;
