By not starting the control server quite so early, it can be reenabled again on all platforms.

diff --git a/Engine/Source/Runtime/TraceLog/Private/Trace/Writer.cpp b/Engine/Source/Runtime/TraceLog/Private/Trace/Writer.cpp
--- a/Engine/Source/Runtime/TraceLog/Private/Trace/Writer.cpp
+++ b/Engine/Source/Runtime/TraceLog/Private/Trace/Writer.cpp
@@ -547,11 +547,6 @@ static bool Writer_ControlDispatch(uint32 ArgC, ANSICHAR const* const* ArgV)
 ////////////////////////////////////////////////////////////////////////////////
 static bool Writer_ControlListen()
 {
-#if PLATFORM_PS4
-	GControlState = EControlState::Failed;
-	return false;
-#endif
-
 	GControlListen = TcpSocketListen(1985);
 	if (!GControlListen)
 	{
@@ -703,8 +698,6 @@ static void Writer_UpdateControl()
 ////////////////////////////////////////////////////////////////////////////////
 static void Writer_InitializeControl()
 {
-	Writer_ControlListen();
-
 	Writer_ControlAddCommand("Connect", nullptr,
 		[] (void*, uint32 ArgC, ANSICHAR const* const* ArgV)
 		{
