When checking if any logged event buffers are ready to be retired, also see if we can see the buffer while it is not in use. If so then what we can see is available to be sent ahead of actual retirement.

diff --git a/Engine/Source/Runtime/TraceLog/Private/Trace/Writer.cpp b/Engine/Source/Runtime/TraceLog/Private/Trace/Writer.cpp
--- a/Engine/Source/Runtime/TraceLog/Private/Trace/Writer.cpp
+++ b/Engine/Source/Runtime/TraceLog/Private/Trace/Writer.cpp
@@ -107,6 +107,7 @@ static FBuffer*						GTailBuffer;
 static FBuffer*						GHeadBuffer;
 static void*						GAllocBase;
 static uint32						GAllocSize;
+static uint32						GTailPreSent;
 
 ////////////////////////////////////////////////////////////////////////////////
 static void Writer_InitializeBuffers()
@@ -135,6 +136,7 @@ static void Writer_InitializeBuffers()
 
 	GTailBuffer = Buffers[0];
 	GHeadBuffer = Buffers[BufferCount - 1];
+	GTailPreSent = 0;
 
 	GActiveBuffer.store(Buffers[0], std::memory_order_release);
 }
@@ -163,7 +165,8 @@ static void Writer_RetireBuffer(void (*DataSink)(const uint8*, uint32))
 
 	if (uint32 SendSize = GTailBuffer->Final - sizeof(FBuffer))
 	{
-		DataSink(GTailBuffer->Data, SendSize);
+		SendSize -= GTailPreSent;
+		DataSink(GTailBuffer->Data + GTailPreSent, SendSize);
 	}
 	uint64 SendTsc = Writer_GetTimestamp();
 
@@ -173,6 +176,7 @@ static void Writer_RetireBuffer(void (*DataSink)(const uint8*, uint32))
 	GHeadBuffer->Next.store(GTailBuffer, std::memory_order_release);
 	GHeadBuffer = GTailBuffer;
 	GTailBuffer = Next;
+	GTailPreSent = 0;
 	uint64 DoneTsc = Writer_GetTimestamp();
 
 	if (UE_TRACE_EVENT_IS_ENABLED($Trace, PerfWorker))
@@ -194,12 +198,23 @@ static void Writer_UpdateBuffers(void (*DataSink)(const uint8*, uint32))
 	for (int i = 0; i < 3; ++i)
 	{
 		FBuffer* Buffer = GActiveBuffer.load(std::memory_order_acquire);
-		if (GTailBuffer == Buffer)
+		if (GTailBuffer != Buffer)
 		{
-			return;
+			Writer_RetireBuffer(DataSink);
+			continue;
+		}
+
+		uint32 Used = Buffer->Used.load(std::memory_order_relaxed);
+		if (Used >= BufferSize) /* is the buffer in use somewhere? */
+		{
+			continue;
 		}
 
-		Writer_RetireBuffer(DataSink);
+		if (uint32 Sendable = Used - sizeof(FBuffer) - GTailPreSent)
+		{
+			DataSink(Buffer->Data + GTailPreSent, Sendable);
+			GTailPreSent += Sendable;
+		}
 	}
 }
 
