Events can now be marked as "Important", indicating that logged events of that type should persist and not be logged ephemerally in the log tail. Note that this is currently semantic only for now and has no effect internally.

diff --git a/Engine/Source/Developer/TraceAnalysis/Private/Analysis/Engine.cpp b/Engine/Source/Developer/TraceAnalysis/Private/Analysis/Engine.cpp
--- a/Engine/Source/Developer/TraceAnalysis/Private/Analysis/Engine.cpp
+++ b/Engine/Source/Developer/TraceAnalysis/Private/Analysis/Engine.cpp
@@ -584,14 +584,15 @@ bool FAnalysisEngine::OnData(FStreamReader::FData& Data)
 			break;
 		}
 
-		if (Header->Uid >= Dispatches.Num())
+		uint16 Uid = uint16(Header->Uid & ((1 << 14) - 1));
+		if (Uid >= Dispatches.Num())
 		{
 			return false;
 		}
 
 		Transport->Advance(BlockSize);
 
-		const FDispatch* Dispatch = Dispatches[Header->Uid];
+		const FDispatch* Dispatch = Dispatches[Uid];
 
 		EventDataImpl->Dispatch = Dispatch;
 		EventDataImpl->Ptr = Header->EventData;
diff --git a/Engine/Source/Runtime/TraceLog/Private/Trace/Writer.cpp b/Engine/Source/Runtime/TraceLog/Private/Trace/Writer.cpp
--- a/Engine/Source/Runtime/TraceLog/Private/Trace/Writer.cpp
+++ b/Engine/Source/Runtime/TraceLog/Private/Trace/Writer.cpp
@@ -89,7 +89,7 @@ void Writer_InitializeTiming()
 {
 	GStartCycle = TimeGetTimestamp();
 
-	UE_TRACE_EVENT_BEGIN($Trace, Timing, Always)
+	UE_TRACE_EVENT_BEGIN($Trace, Timing, Always|Important)
 		UE_TRACE_EVENT_FIELD(uint64, StartCycle)
 		UE_TRACE_EVENT_FIELD(uint64, CycleFrequency)
 	UE_TRACE_EVENT_END()
@@ -774,7 +774,7 @@ static bool GInitialized = false;
 ////////////////////////////////////////////////////////////////////////////////
 static void Writer_LogHeader()
 {
-	UE_TRACE_EVENT_BEGIN($Trace, NewTrace, Always)
+	UE_TRACE_EVENT_BEGIN($Trace, NewTrace, Always|Important)
 		UE_TRACE_EVENT_FIELD(uint16, Endian)
 		UE_TRACE_EVENT_FIELD(uint8, Version)
 		UE_TRACE_EVENT_FIELD(uint8, PointerSize)
@@ -887,7 +887,10 @@ enum class EKnownEventUids : uint16
 	NewEvent		= FNewEventEvent::Uid,
 	User,
 	Max				= (1 << 14) - 1, // ...leaves two MSB bits for other uses.
+	UidMask			= Max,
 	Invalid			= Max,
+	Flag_Unused		= 1 << 14,
+	Flag_Important	= 1 << 15,
 };
 
 ////////////////////////////////////////////////////////////////////////////////
@@ -934,6 +937,11 @@ void Writer_EventCreate(
 		return;
 	}
 
+	if (Flags & FEvent::Flag_Important)
+	{
+		Uid |= uint16(EKnownEventUids::Flag_Important);
+	}
+
 	uint32 LoggerHash = Writer_EventGetHash(LoggerName.Ptr);
 	uint32 NameHash = Writer_EventGetHash(EventName.Ptr);
 
@@ -959,7 +967,7 @@ void Writer_EventCreate(
 	auto& Event = *(FNewEventEvent*)Writer_BeginLog(uint16(EKnownEventUids::NewEvent), EventSize);
 
 	// Write event's main properties.
-	Event.EventUid = Target->Uid;
+	Event.EventUid = uint16(Uid) & uint16(EKnownEventUids::UidMask);
 	Event.LoggerNameSize = LoggerName.Length;
 	Event.EventNameSize = EventName.Length;
 
diff --git a/Engine/Source/Runtime/TraceLog/Public/Trace/Private/Event.h b/Engine/Source/Runtime/TraceLog/Public/Trace/Private/Event.h
--- a/Engine/Source/Runtime/TraceLog/Public/Trace/Private/Event.h
+++ b/Engine/Source/Runtime/TraceLog/Public/Trace/Private/Event.h
@@ -46,7 +46,8 @@ public:
 
 	enum
 	{
-		Flag_Always	= 1 << 0,
+		Flag_Always		= 1 << 0,
+		Flag_Important	= 1 << 1,
 	};
 
 	class FLogScope
diff --git a/Engine/Source/Runtime/TraceLog/Public/Trace/Private/Trace.h b/Engine/Source/Runtime/TraceLog/Public/Trace/Private/Trace.h
--- a/Engine/Source/Runtime/TraceLog/Public/Trace/Private/Trace.h
+++ b/Engine/Source/Runtime/TraceLog/Public/Trace/Private/Trace.h
@@ -47,6 +47,7 @@
 			static const bool bOnceOnly = [] () \
 			{ \
 				const uint32 Always = Trace::FEvent::Flag_Always; \
+				const uint32 Important = Trace::FEvent::Flag_Important; \
 				const uint32 Flags = (0, ##__VA_ARGS__); \
 				F##LoggerName##EventName##Fields Fields; \
 				const auto* Descs = (Trace::FFieldDesc*)(&Fields); \
