The hold buffer is used to store a finite amoung of trace events until a connection to a recorder is establish. Due to a logic error it would not use its last page page, and it's only one page big these days. Important events were getting lost and were unanalyssaables.

diff --git a/Engine/Source/Runtime/TraceLog/Private/Trace/Writer.cpp b/Engine/Source/Runtime/TraceLog/Private/Trace/Writer.cpp
--- a/Engine/Source/Runtime/TraceLog/Private/Trace/Writer.cpp
+++ b/Engine/Source/Runtime/TraceLog/Private/Trace/Writer.cpp
@@ -366,7 +366,7 @@ void FHoldBuffer::Write(const void* Data, uint32 Size)
 	uint16 HotPageCount = uint16((NextUsed + (FHoldBuffer::PageSize - 1)) >> FHoldBuffer::PageShift);
 	if (HotPageCount > MappedPageCount)
 	{
-		if (HotPageCount >= FHoldBuffer::MaxPages)
+		if (HotPageCount > FHoldBuffer::MaxPages)
 		{
 			bFull = true;
 			return;
