Fixed "no debug symbols in executable" warnings on Mac.

diff --git a/Engine/Source/Runtime/TraceLog/Private/Trace/Trace.cpp b/Engine/Source/Runtime/TraceLog/Private/Trace/Trace.cpp
--- a/Engine/Source/Runtime/TraceLog/Private/Trace/Trace.cpp
+++ b/Engine/Source/Runtime/TraceLog/Private/Trace/Trace.cpp
@@ -61,4 +61,10 @@ void Flush()
 
 } // namespace Trace
 
+#else
+
+#if PLATFORM_MAC
+void Dummy() {}
+#endif // PLATFORM_MAC
+
 #endif // UE_TRACE_ENABLED
