Renamed Trace's event object so that it doesn't cause ambiguation.

diff --git a/Engine/Source/Developer/TraceAnalysis/Private/Analysis/Engine.cpp b/Engine/Source/Developer/TraceAnalysis/Private/Analysis/Engine.cpp
--- a/Engine/Source/Developer/TraceAnalysis/Private/Analysis/Engine.cpp
+++ b/Engine/Source/Developer/TraceAnalysis/Private/Analysis/Engine.cpp
@@ -9,7 +9,7 @@
 #if 0
 #include "Trace/Private/DataDecoder.h"
 #endif // 0
-#include "Trace/Private/Event.h"
+#include "Trace/Private/EventDef.h"
 
 namespace Trace
 {
diff --git a/Engine/Source/Runtime/TraceLog/Private/Trace/Trace.cpp b/Engine/Source/Runtime/TraceLog/Private/Trace/Trace.cpp
--- a/Engine/Source/Runtime/TraceLog/Private/Trace/Trace.cpp
+++ b/Engine/Source/Runtime/TraceLog/Private/Trace/Trace.cpp
@@ -1,7 +1,7 @@
 // Copyright 1998-2019 Epic Games, Inc. All Rights Reserved.
 
 #include "Trace/Trace.h"
-#include "Trace/Private/Event.h"
+#include "Trace/Private/EventDef.h"
 
 #if UE_TRACE_ENABLED
 
diff --git a/Engine/Source/Runtime/TraceLog/Private/Trace/Writer.cpp b/Engine/Source/Runtime/TraceLog/Private/Trace/Writer.cpp
--- a/Engine/Source/Runtime/TraceLog/Private/Trace/Writer.cpp
+++ b/Engine/Source/Runtime/TraceLog/Private/Trace/Writer.cpp
@@ -181,7 +181,7 @@ static void Writer_RetireBuffer(void (*DataSink)(const uint8*, uint32))
 
 	if (UE_TRACE_EVENT_IS_ENABLED($Trace, PerfWorker))
 	{
-		uint8 PerfEventBuffer[TRACE_PRIVATE_EVENT_SIZE($Trace, PerfWorker) + FEvent::HeaderSize];
+		uint8 PerfEventBuffer[TRACE_PRIVATE_EVENT_SIZE($Trace, PerfWorker) + FEventDef::HeaderSize];
 		UE_TRACE_LOG($Trace, PerfWorker, PerfEventBuffer)
 			<< PerfWorker.Start(StartTsc)
 			<< PerfWorker.Acquire(uint32(AcquireTailTsc - StartTsc))
@@ -251,7 +251,7 @@ UE_TRACE_API void* Writer_NextBuffer(Private::FBuffer* Buffer, uint32 PrevUsed,
 	uint32 PerfEventSize = 0;
 	if (UE_TRACE_EVENT_IS_ENABLED($Trace, PerfNextBuffer))
 	{
-		PerfEventSize = TRACE_PRIVATE_EVENT_SIZE($Trace, PerfNextBuffer) + FEvent::HeaderSize;
+		PerfEventSize = TRACE_PRIVATE_EVENT_SIZE($Trace, PerfNextBuffer) + FEventDef::HeaderSize;
 	}
 	SizeAndRef += PerfEventSize;
 
@@ -878,8 +878,8 @@ struct TLateAtomic
 	alignas(InnerType) char Buffer[sizeof(InnerType)];
 };
 
-static TLateAtomic<uint32>	GEventUidCounter;	// = 0;
-static TLateAtomic<FEvent*>	GHeadEvent;			// = nullptr;
+static TLateAtomic<uint32>		GEventUidCounter;	// = 0;
+static TLateAtomic<FEventDef*>	GHeadEvent;			// = nullptr;
 
 ////////////////////////////////////////////////////////////////////////////////
 enum class EKnownEventUids : uint16
@@ -915,7 +915,7 @@ static uint32 Writer_EventGetHash(uint32 LoggerHash, uint32 NameHash)
 
 ////////////////////////////////////////////////////////////////////////////////
 void Writer_EventCreate(
-	FEvent* Target,
+	FEventDef* Target,
 	const FLiteralName& LoggerName,
 	const FLiteralName& EventName,
 	const FFieldDesc* FieldDescs,
@@ -937,7 +937,7 @@ void Writer_EventCreate(
 		return;
 	}
 
-	if (Flags & FEvent::Flag_Important)
+	if (Flags & FEventDef::Flag_Important)
 	{
 		Uid |= uint16(EKnownEventUids::Flag_Important);
 	}
@@ -949,7 +949,7 @@ void Writer_EventCreate(
  	Target->Uid = uint16(Uid);
 	Target->LoggerHash = LoggerHash;
 	Target->Hash = Writer_EventGetHash(LoggerHash, NameHash);
-	Target->Enabled.Internal = !!(Flags & FEvent::Flag_Always);
+	Target->Enabled.Internal = !!(Flags & FEventDef::Flag_Always);
 	Target->Enabled.bOptedIn = false;
 	Target->bInitialized = true;
 
@@ -1004,7 +1004,7 @@ void Writer_EventCreate(
 	// Add this new event into the list so we can look them up later.
 	while (true)
 	{
-		FEvent* HeadEvent = GHeadEvent->load(std::memory_order_relaxed);
+		FEventDef* HeadEvent = GHeadEvent->load(std::memory_order_relaxed);
 		Target->Handle = HeadEvent;
 		if (GHeadEvent->compare_exchange_weak(HeadEvent, Target, std::memory_order_release))
 		{
@@ -1025,8 +1025,8 @@ uint32 Writer_EventToggle(const ANSICHAR* Wildcard, bool bState)
 	{
 		uint32 LoggerHash = Writer_EventGetHash(Wildcard);
 
-		FEvent* Event = Private::GHeadEvent->load(std::memory_order_relaxed);
-		for (; Event != nullptr; Event = (FEvent*)(Event->Handle))
+		FEventDef* Event = Private::GHeadEvent->load(std::memory_order_relaxed);
+		for (; Event != nullptr; Event = (FEventDef*)(Event->Handle))
 		{
 			if (Event->LoggerHash == LoggerHash)
 			{
@@ -1042,8 +1042,8 @@ uint32 Writer_EventToggle(const ANSICHAR* Wildcard, bool bState)
 	uint32 NameHash = Writer_EventGetHash(Dot + 1);
 	uint32 EventHash = Writer_EventGetHash(LoggerHash, NameHash);
 
-	FEvent* Event = Private::GHeadEvent->load(std::memory_order_relaxed);
-	for (; Event != nullptr; Event = (FEvent*)(Event->Handle))
+	FEventDef* Event = Private::GHeadEvent->load(std::memory_order_relaxed);
+	for (; Event != nullptr; Event = (FEventDef*)(Event->Handle))
 	{
 		if (Event->Hash == EventHash)
 		{
diff --git a/Engine/Source/Runtime/TraceLog/Public/Trace/Private/Event.h b/Engine/Source/Runtime/TraceLog/Public/Trace/Private/Event.h
deleted file mode 100644
diff --git a/Engine/Source/Runtime/TraceLog/Public/Trace/Private/Event.inl b/Engine/Source/Runtime/TraceLog/Public/Trace/Private/Event.inl
deleted file mode 100644
diff --git a/Engine/Source/Runtime/TraceLog/Public/Trace/Private/EventDef.h b/Engine/Source/Runtime/TraceLog/Public/Trace/Private/EventDef.h
new file mode 100644
--- /dev/null
+++ b/Engine/Source/Runtime/TraceLog/Public/Trace/Private/EventDef.h
@@ -0,0 +1,83 @@
+// Copyright 1998-2019 Epic Games, Inc. All Rights Reserved.
+
+#pragma once
+
+#include "Trace.h"
+
+namespace Trace
+{
+
+////////////////////////////////////////////////////////////////////////////////
+struct FNewEventEvent
+{
+	enum : uint16 { Uid = 0 };
+
+	uint16		EventUid;
+	uint16		FieldCount;
+	uint8		LoggerNameSize;
+	uint8		EventNameSize;
+	struct
+	{
+		uint16	Offset;
+		uint16	Size;
+		uint8	TypeInfo;
+		uint8	NameSize;
+	}			Fields[];
+	/*uint8		NameData[]*/
+};
+
+} // namespace Trace
+
+
+
+#if UE_TRACE_ENABLED
+
+namespace Trace
+{
+
+struct FFieldDesc;
+struct FLiteralName;
+
+////////////////////////////////////////////////////////////////////////////////
+class FEventDef
+{
+public:
+	enum { HeaderSize = sizeof(uint32), }; // I'll just dump this here for now.
+
+	enum
+	{
+		Flag_Always		= 1 << 0,
+		Flag_Important	= 1 << 1,
+	};
+
+	class FLogScope
+	{
+	public:
+						FLogScope(uint16 EventUid, uint16 Size);
+						FLogScope(uint16 EventUid, uint16 Size, uint16 ExtraBytes);
+						FLogScope(uint16 EventUid, uint16 Size, uint8* Out);
+						~FLogScope();
+		uint8*			Ptr;
+		const bool		bOutOfBand = false;
+	};
+
+	void*						Handle;
+	uint32						LoggerHash;
+	uint32						Hash;
+	uint16						Uid;
+	union
+	{
+		struct
+		{
+			bool				bOptedIn;
+			uint8				Internal;
+		};
+		uint16					Test;
+	}							Enabled;
+	bool						bInitialized;
+	UE_TRACE_API static void	Create(FEventDef* Target, const FLiteralName& LoggerName, const FLiteralName& EventName, const FFieldDesc* FieldDescs, uint32 FieldCount, uint32 Flags=0);
+};
+
+} // namespace Trace
+
+#endif // UE_TRACE_ENABLED
diff --git a/Engine/Source/Runtime/TraceLog/Public/Trace/Private/EventDef.inl b/Engine/Source/Runtime/TraceLog/Public/Trace/Private/EventDef.inl
new file mode 100644
--- /dev/null
+++ b/Engine/Source/Runtime/TraceLog/Public/Trace/Private/EventDef.inl
@@ -0,0 +1,52 @@
+// Copyright 1998-2019 Epic Games, Inc. All Rights Reserved.
+
+#pragma once
+
+#if UE_TRACE_ENABLED
+
+#include "EventDef.h"
+#include "Writer.inl"
+
+namespace Trace
+{
+
+////////////////////////////////////////////////////////////////////////////////
+inline FEventDef::FLogScope::FLogScope(uint16 EventUid, uint16 Size)
+{
+	Ptr = Writer_BeginLog(EventUid, Size);
+}
+
+////////////////////////////////////////////////////////////////////////////////
+inline FEventDef::FLogScope::FLogScope(uint16 EventUid, uint16 Size, uint16 ExtraBytes)
+{
+	Ptr = Writer_BeginLog(EventUid, Size + ExtraBytes);
+}
+
+////////////////////////////////////////////////////////////////////////////////
+inline FEventDef::FLogScope::FLogScope(uint16 EventUid, uint16 Size, uint8* Out)
+: bOutOfBand(true)
+{
+	*(uint32*)Out = (uint32(Size) << 16)|uint32(EventUid);
+	Ptr = Out + HeaderSize;
+}
+
+////////////////////////////////////////////////////////////////////////////////
+inline FEventDef::FLogScope::~FLogScope()
+{
+	if (!bOutOfBand)
+		Writer_EndLog(Ptr);
+}
+
+
+
+////////////////////////////////////////////////////////////////////////////////
+template <typename ActionType>
+inline const FEventDef::FLogScope& operator << (const FEventDef::FLogScope& Lhs, const ActionType& Rhs)
+{
+	Rhs.Write(Lhs.Ptr);
+	return Lhs;
+}
+
+} // namespace Trace
+
+#endif // UE_TRACE_ENABLED
diff --git a/Engine/Source/Runtime/TraceLog/Public/Trace/Private/Trace.h b/Engine/Source/Runtime/TraceLog/Public/Trace/Private/Trace.h
--- a/Engine/Source/Runtime/TraceLog/Public/Trace/Private/Trace.h
+++ b/Engine/Source/Runtime/TraceLog/Public/Trace/Private/Trace.h
@@ -30,7 +30,7 @@
 #if UE_TRACE_ENABLED
 
 #define TRACE_PRIVATE_EVENT_DEFINE(LoggerName, EventName) \
-	Trace::FEvent LoggerName##EventName##Event;
+	Trace::FEventDef LoggerName##EventName##Event;
 
 #define TRACE_PRIVATE_EVENT_BEGIN(LoggerName, EventName, ...) \
 	TRACE_PRIVATE_EVENT_BEGIN_IMPL(static, LoggerName, EventName, ##__VA_ARGS__)
@@ -46,15 +46,15 @@
 		{ \
 			static const bool bOnceOnly = [] () \
 			{ \
-				const uint32 Always = Trace::FEvent::Flag_Always; \
-				const uint32 Important = Trace::FEvent::Flag_Important; \
+				const uint32 Always = Trace::FEventDef::Flag_Always; \
+				const uint32 Important = Trace::FEventDef::Flag_Important; \
 				const uint32 Flags = (0, ##__VA_ARGS__); \
 				F##LoggerName##EventName##Fields Fields; \
 				const auto* Descs = (Trace::FFieldDesc*)(&Fields); \
 				uint32 DescCount = uint32(sizeof(Fields) / sizeof(*Descs)); \
 				const auto& LoggerLiteral = Trace::FLiteralName(#LoggerName); \
 				const auto& EventLiteral = Trace::FLiteralName(#EventName); \
-				Trace::FEvent::Create(&LoggerName##EventName##Event, LoggerLiteral, EventLiteral, Descs, DescCount, Flags); \
+				Trace::FEventDef::Create(&LoggerName##EventName##Event, LoggerLiteral, EventLiteral, Descs, DescCount, Flags); \
 				return true; \
 			}(); \
 		} \
@@ -82,7 +82,7 @@
 #define TRACE_PRIVATE_LOG(LoggerName, EventName, ...) \
 	if (TRACE_PRIVATE_EVENT_IS_ENABLED(LoggerName, EventName)) \
 		if (const auto& __restrict EventName = (F##LoggerName##EventName##Fields&)LoggerName##EventName##Event) \
-			Trace::FEvent::FLogScope(LoggerName##EventName##Event.Uid, TRACE_PRIVATE_EVENT_SIZE(LoggerName, EventName), ##__VA_ARGS__)
+			Trace::FEventDef::FLogScope(LoggerName##EventName##Event.Uid, TRACE_PRIVATE_EVENT_SIZE(LoggerName, EventName), ##__VA_ARGS__)
 
 #else
 
