Corrected preprocessor statements to fix non-MS platforms.

diff --git a/Engine/Source/Runtime/TraceLog/Public/Trace/Private/EventDef.h b/Engine/Source/Runtime/TraceLog/Public/Trace/Private/EventDef.h
--- a/Engine/Source/Runtime/TraceLog/Public/Trace/Private/EventDef.h
+++ b/Engine/Source/Runtime/TraceLog/Public/Trace/Private/EventDef.h
@@ -7,7 +7,7 @@
 namespace Trace
 {
 
-#if _MSC_VER
+#if defined(_MSC_VER)
 	#pragma warning(push)
 	#pragma warning(disable : 4200) // non-standard zero-sized array
 #endif
@@ -31,7 +31,7 @@ struct FNewEventEvent
 	/*uint8		NameData[]*/
 };
 
-#if _MSC_VER
+#if defined(_MSC_VER)
 	#pragma warning(pop)
 #endif
 
diff --git a/Engine/Source/Runtime/TraceLog/Public/Trace/Private/Writer.inl b/Engine/Source/Runtime/TraceLog/Public/Trace/Private/Writer.inl
--- a/Engine/Source/Runtime/TraceLog/Public/Trace/Private/Writer.inl
+++ b/Engine/Source/Runtime/TraceLog/Public/Trace/Private/Writer.inl
@@ -16,7 +16,7 @@ using TTraceAtomic = std::atomic<Type>;
 namespace Private
 {
 
-#if _MSC_VER
+#if defined(_MSC_VER)
 	#pragma warning(push)
 	#pragma warning(disable : 4200) // non-standard zero-sized array
 #endif
@@ -44,7 +44,7 @@ struct FBuffer
 	alignas(void*) uint8		Data[];
 };
 
-#if _MSC_VER
+#if defined(_MSC_VER)
 	#pragma warning(pop)
 #endif
 
