File missed from earlier commit.

diff --git a/Engine/Source/Runtime/TraceLog/Public/Trace/Trace.h b/Engine/Source/Runtime/TraceLog/Public/Trace/Trace.h
--- a/Engine/Source/Runtime/TraceLog/Public/Trace/Trace.h
+++ b/Engine/Source/Runtime/TraceLog/Public/Trace/Trace.h
@@ -24,5 +24,5 @@ UE_TRACE_API void	Flush() UE_TRACE_IMPL();
 #define UE_TRACE_LOG(LoggerName, EventName, ...)				TRACE_PRIVATE_LOG(LoggerName, EventName, ##__VA_ARGS__)
 
 #include "Private/Field.h"
-#include "Private/Event.inl"
+#include "Private/EventDef.inl"
 #include "Private/Writer.inl"
