Another missed file. Time for a reconcile...

diff --git a/Engine/Source/Runtime/TraceLog/Private/Trace/Event.cpp b/Engine/Source/Runtime/TraceLog/Private/Trace/Event.cpp
deleted file mode 100644
diff --git a/Engine/Source/Runtime/TraceLog/Private/Trace/EventDef.cpp b/Engine/Source/Runtime/TraceLog/Private/Trace/EventDef.cpp
new file mode 100644
--- /dev/null
+++ b/Engine/Source/Runtime/TraceLog/Private/Trace/EventDef.cpp
@@ -0,0 +1,37 @@
+// Copyright 1998-2019 Epic Games, Inc. All Rights Reserved.
+
+#include "Trace/Private/EventDef.h"
+
+#if UE_TRACE_ENABLED
+
+#include "Trace/Private/Field.h"
+#include "Trace/Private/Writer.inl"
+
+namespace Trace
+{
+
+////////////////////////////////////////////////////////////////////////////////
+namespace Private
+{
+
+void Writer_EventCreate(FEventDef*, const FLiteralName&, const FLiteralName&, const FFieldDesc*, uint32, uint32);
+
+} // namespace Private
+
+
+
+////////////////////////////////////////////////////////////////////////////////
+void FEventDef::Create(
+	FEventDef* Target,
+	const FLiteralName& LoggerName,
+	const FLiteralName& EventName,
+	const FFieldDesc* FieldDescs,
+	uint32 FieldCount,
+	uint32 Flags)
+{
+	Private::Writer_EventCreate(Target, LoggerName, EventName, FieldDescs, FieldCount, Flags);
+}
+
+} // namespace Trace
+
+#endif // UE_TRACE_ENABLED
