Added an enum entry for masking of a trace event's field's category.

diff --git a/Engine/Source/Runtime/TraceLog/Public/Trace/Private/Field.h b/Engine/Source/Runtime/TraceLog/Public/Trace/Private/Field.h
--- a/Engine/Source/Runtime/TraceLog/Public/Trace/Private/Field.h
+++ b/Engine/Source/Runtime/TraceLog/Public/Trace/Private/Field.h
@@ -11,6 +11,7 @@ namespace Trace
 enum : uint8
 {
 	/* Category */
+	_Field_CategoryMask	= 0300,
 	_Field_Integer		= 0000, 
 	_Field_Float		= 0100,
 	_Field_Array		= 0200, /* flag */
