Renamed trace's detail files so the file names identify their platforms more clearly.

diff --git a/Engine/Source/Runtime/TraceLog/Private/Trace/Detail/Windows/WindowsTrace.cpp b/Engine/Source/Runtime/TraceLog/Private/Trace/Detail/Windows/WindowsTrace.cpp
new file mode 100644
--- /dev/null
+++ b/Engine/Source/Runtime/TraceLog/Private/Trace/Detail/Windows/WindowsTrace.cpp
@@ -0,0 +1,244 @@
+// Copyright 1998-2019 Epic Games, Inc. All Rights Reserved.
+
+#include "Trace/Trace.h"
+
+#if UE_TRACE_ENABLED
+
+#include "Windows/AllowWindowsPlatformTypes.h"
+#	include "Windows/WindowsHWrapper.h"
+#	define _WINSOCK_DEPRECATED_NO_WARNINGS  
+#	include <winsock2.h>
+#	include <ws2tcpip.h>
+#	pragma comment(lib, "ws2_32.lib")
+#include "Windows/HideWindowsPlatformTypes.h"
+
+#pragma warning(push)
+#pragma warning(disable : 6250) // VirtualFree() missing MEM_RELEASE - a false positive
+#pragma warning(disable : 6031) // WSAStartup() return ignore  - we're error tolerant
+
+namespace Trace
+{
+
+////////////////////////////////////////////////////////////////////////////////
+uint8* MemoryReserve(SIZE_T Size)
+{
+	return (uint8*)VirtualAlloc(nullptr, Size, MEM_RESERVE, PAGE_READWRITE);
+}
+
+////////////////////////////////////////////////////////////////////////////////
+void MemoryFree(void* Address, SIZE_T Size)
+{
+	VirtualFree(Address, 0, MEM_RELEASE);
+}
+
+////////////////////////////////////////////////////////////////////////////////
+void MemoryMap(void* Address, SIZE_T Size)
+{
+	auto Inner = [] (void* Address, SIZE_T Size) -> void*
+	{
+		return VirtualAlloc(Address, Size, MEM_COMMIT, PAGE_READWRITE);
+	};
+	Inner(Address, Size);
+}
+
+////////////////////////////////////////////////////////////////////////////////
+void MemoryUnmap(void* Address, SIZE_T Size)
+{
+	VirtualFree(Address, Size, MEM_DECOMMIT);
+}
+
+
+
+////////////////////////////////////////////////////////////////////////////////
+UPTRINT ThreadCreate(const ANSICHAR* Name, void (*Entry)())
+{
+	DWORD (WINAPI *WinApiThunk)(void*) = [] (void* Param) -> DWORD
+	{
+		typedef void (*EntryType)(void);
+		(EntryType(Param))();
+		return 0;
+	};
+
+	HANDLE Handle = CreateThread(nullptr, 0, WinApiThunk, (void*)Entry, 0, nullptr);
+	return UPTRINT(Handle);
+}
+
+////////////////////////////////////////////////////////////////////////////////
+uint32 ThreadGetCurrentId()
+{
+	return GetCurrentThreadId();
+}
+
+////////////////////////////////////////////////////////////////////////////////
+void ThreadSleep(uint32 Milliseconds)
+{
+	Sleep(Milliseconds);
+}
+
+////////////////////////////////////////////////////////////////////////////////
+void ThreadJoin(UPTRINT Handle)
+{
+	WaitForSingleObject(HANDLE(Handle), INFINITE);
+}
+
+////////////////////////////////////////////////////////////////////////////////
+void ThreadDestroy(UPTRINT Handle)
+{
+	CloseHandle(HANDLE(Handle));
+}
+
+
+
+////////////////////////////////////////////////////////////////////////////////
+uint64 TimeGetFrequency()
+{
+	LARGE_INTEGER Value;
+	QueryPerformanceFrequency(&Value);
+	return Value.QuadPart;
+}
+
+////////////////////////////////////////////////////////////////////////////////
+uint64 TimeGetTimestamp()
+{
+	LARGE_INTEGER Value;
+	QueryPerformanceCounter(&Value);
+	return Value.QuadPart;
+}
+
+
+
+////////////////////////////////////////////////////////////////////////////////
+static void TcpSocketInitialize()
+{
+	WSADATA WsaData;
+	WSAStartup(MAKEWORD(2, 2), &WsaData);
+}
+
+////////////////////////////////////////////////////////////////////////////////
+UPTRINT TcpSocketConnect(const ANSICHAR* Host, uint16 Port)
+{
+	TcpSocketInitialize();
+
+	struct FAddrInfoPtr
+	{
+					~FAddrInfoPtr()	{ freeaddrinfo(Value); }
+		addrinfo*	operator -> ()	{ return Value; }
+		addrinfo**	operator & ()	{ return &Value; }
+		addrinfo*	Value;
+	};
+
+	FAddrInfoPtr Info;
+	addrinfo Hints = {};
+	Hints.ai_family = AF_INET;
+	Hints.ai_socktype = SOCK_STREAM;
+	Hints.ai_protocol = IPPROTO_TCP;
+	if (getaddrinfo(Host, nullptr, &Hints, &Info))
+	{
+		return 0;
+	}
+
+	if (&Info == nullptr)
+	{
+		return 0;
+	}
+
+	auto* SockAddr = (sockaddr_in*)Info->ai_addr;
+	SockAddr->sin_port = htons(Port);
+
+	SOCKET Socket = socket(AF_INET, SOCK_STREAM, IPPROTO_TCP);
+	if (Socket == INVALID_SOCKET)
+	{
+		return 0;
+	}
+
+	int Result = connect(Socket, Info->ai_addr, int(Info->ai_addrlen));
+	if (Result == SOCKET_ERROR)
+	{
+		closesocket(Socket);
+		return 0;
+	}
+
+	return UPTRINT(Socket + 1);
+}
+
+////////////////////////////////////////////////////////////////////////////////
+UPTRINT TcpSocketListen(uint16 Port)
+{
+	TcpSocketInitialize();
+
+	SOCKET Socket = socket(AF_INET, SOCK_STREAM, IPPROTO_TCP);
+	if (Socket == INVALID_SOCKET)
+	{
+		return 0;
+	}
+
+	sockaddr_in SockAddr;
+	SockAddr.sin_family = AF_INET;
+	SockAddr.sin_addr.s_addr = 0;
+	SockAddr.sin_port = htons(Port);
+	int Result = bind(Socket, (SOCKADDR*)&SockAddr, sizeof(SockAddr));
+	if (Result == INVALID_SOCKET)
+	{
+		closesocket(Socket);
+		return 0;
+	}
+
+	Result = listen(Socket, 1);
+	if (Result == INVALID_SOCKET)
+	{
+		closesocket(Socket);
+		return 0;
+	}
+
+	return UPTRINT(Socket + 1);
+}
+
+////////////////////////////////////////////////////////////////////////////////
+UPTRINT TcpSocketAccept(UPTRINT Socket)
+{
+	SOCKET Inner = Socket - 1;
+
+	Inner = accept(Inner, nullptr, nullptr);
+	if (Inner == INVALID_SOCKET)
+	{
+		return 0;
+	}
+
+	return UPTRINT(Inner + 1);
+}
+
+////////////////////////////////////////////////////////////////////////////////
+void TcpSocketClose(UPTRINT Socket)
+{
+	SOCKET Inner = Socket - 1;
+	closesocket(Inner);
+}
+
+////////////////////////////////////////////////////////////////////////////////
+bool TcpSocketSelect(UPTRINT Socket)
+{
+	SOCKET Inner = Socket - 1;
+	fd_set FdSet = { 1, { Inner }, };
+	TIMEVAL TimeVal = {};
+	return (select(0, &FdSet, nullptr, nullptr, &TimeVal) != 0);
+}
+
+////////////////////////////////////////////////////////////////////////////////
+bool TcpSocketSend(UPTRINT Socket, const void* Data, uint32 Size)
+{
+	SOCKET Inner = Socket - 1;
+	return send(Inner, (const char*)Data, Size, 0) == Size;
+}
+
+////////////////////////////////////////////////////////////////////////////////
+int32 TcpSocketRecv(UPTRINT Socket, void* Data, uint32 Size)
+{
+	SOCKET Inner = Socket - 1;
+	return recv(Inner, (char*)Data, Size, 0);
+}
+
+} // namespace Trace
+
+#pragma warning(pop)
+
+#endif // UE_TRACE_ENABLED
diff --git a/Engine/Source/Runtime/TraceLog/Private/Trace/Windows/Detail.cpp b/Engine/Source/Runtime/TraceLog/Private/Trace/Windows/Detail.cpp
deleted file mode 100644
