_mm_pause() equivalent for Arm.

diff --git a/Engine/Source/Runtime/TraceLog/Private/Trace/Writer.cpp b/Engine/Source/Runtime/TraceLog/Private/Trace/Writer.cpp
--- a/Engine/Source/Runtime/TraceLog/Private/Trace/Writer.cpp
+++ b/Engine/Source/Runtime/TraceLog/Private/Trace/Writer.cpp
@@ -7,7 +7,9 @@
 #include "Trace/Trace.h"
 #include "Misc/CString.h"
 
-#include <emmintrin.h>
+#if PLATFORM_CPU_X86_FAMILY
+	#include <emmintrin.h>
+#endif
 
 namespace Trace
 {
@@ -41,8 +43,10 @@ inline void Writer_Yield()
 {
 #if PLATFORM_CPU_X86_FAMILY
 	_mm_pause();
+#elif PLATFORM_CPU_ARM_FAMILY
+	__builtin_arm_yield();
 #else
-#	error Unsupported platform!
+	#error Unsupported platform!
 #endif
 }
 
