First pass at implementing the Linux/Posix parts of trace.

diff --git a/Engine/Source/Runtime/TraceLog/Private/Trace/Detail/Linux/LinuxTrace.cpp b/Engine/Source/Runtime/TraceLog/Private/Trace/Detail/Linux/LinuxTrace.cpp
new file mode 100644
--- /dev/null
+++ b/Engine/Source/Runtime/TraceLog/Private/Trace/Detail/Linux/LinuxTrace.cpp
@@ -0,0 +1,57 @@
+// Copyright 1998-2019 Epic Games, Inc. All Rights Reserved.
+
+#include "Trace/Trace.h"
+
+#if UE_TRACE_ENABLED
+
+#include "Trace/Detail/PosixTrace.inl"
+
+#include <sys/mman.h>
+#include <time.h>
+
+namespace Trace
+{
+
+////////////////////////////////////////////////////////////////////////////////
+uint8* MemoryReserve(SIZE_T Size)
+{
+	return (uint8*)mmap(nullptr, Size, PROT_NONE, MAP_PRIVATE|MAP_ANONYMOUS, -1, 0);
+}
+
+////////////////////////////////////////////////////////////////////////////////
+void MemoryFree(void* Address, SIZE_T Size)
+{
+	munmap(Address, Size);
+}
+
+////////////////////////////////////////////////////////////////////////////////
+void MemoryMap(void* Address, SIZE_T Size)
+{
+	mprotect(Address, Size, PROT_READ|PROT_WRITE);
+}
+
+////////////////////////////////////////////////////////////////////////////////
+void MemoryUnmap(void* Address, SIZE_T Size)
+{
+	madvise(Address, Size, MADV_REMOVE);
+}
+
+
+
+////////////////////////////////////////////////////////////////////////////////
+uint64 TimeGetFrequency()
+{
+	return uint64(1000 * 1000 * 1000);
+}
+
+////////////////////////////////////////////////////////////////////////////////
+uint64 TimeGetTimestamp()
+{
+	timespec Time;
+	clock_gettime(CLOCK_MONOTONIC_RAW, &Time);
+	return (uint64(Time.tv_sec) * (1000 * 1000 * 1000)) + Time.tv_nsec;
+}
+
+} // namespace Trace
+
+#endif // UE_TRACE_ENABLED
diff --git a/Engine/Source/Runtime/TraceLog/Private/Trace/Detail/PosixTrace.inl b/Engine/Source/Runtime/TraceLog/Private/Trace/Detail/PosixTrace.inl
new file mode 100644
--- /dev/null
+++ b/Engine/Source/Runtime/TraceLog/Private/Trace/Detail/PosixTrace.inl
@@ -0,0 +1,158 @@
+// Copyright 1998-2019 Epic Games, Inc. All Rights Reserved.
+
+#include <pthread.h>
+#include <unistd.h>
+
+#include <arpa/inet.h>
+#include <sys/socket.h>
+
+namespace Trace
+{
+
+////////////////////////////////////////////////////////////////////////////////
+UPTRINT ThreadCreate(const ANSICHAR* Name, void (*Entry)())
+{
+	auto Thunk = [] (void* Param) -> void*
+	{
+		typedef void (*EntryType)(void);
+		(EntryType(Param))();
+		return nullptr;
+	};
+
+	pthread_t Thread;
+	if (pthread_create(&Thread, nullptr, Thunk, (void*)Entry) != 0)
+	{
+		return 0;
+	}
+
+	return UPTRINT(Thread);
+}
+
+////////////////////////////////////////////////////////////////////////////////
+uint32 ThreadGetCurrentId()
+{
+	return pthread_self();
+}
+
+////////////////////////////////////////////////////////////////////////////////
+void ThreadSleep(uint32 Milliseconds)
+{
+	sleep(Milliseconds << 10);
+}
+
+////////////////////////////////////////////////////////////////////////////////
+void ThreadJoin(UPTRINT Handle)
+{
+	pthread_t Inner = pthread_t(Handle);
+	pthread_join(Inner, nullptr);
+}
+
+////////////////////////////////////////////////////////////////////////////////
+void ThreadDestroy(UPTRINT Handle)
+{
+	pthread_t Inner = pthread_t(Handle);
+	pthread_detach(Inner);
+}
+
+
+
+////////////////////////////////////////////////////////////////////////////////
+UPTRINT TcpSocketConnect(const ANSICHAR* Host, uint16 Port)
+{
+	sockaddr_in Addr;
+	Addr.sin_family = AF_INET;
+	Addr.sin_port = htons(Port);
+	inet_aton(Host, &Addr.sin_addr);
+
+	int Socket = socket(AF_INET, SOCK_STREAM, 0);
+	if (Socket < 0)
+	{
+		return 0;
+	}
+
+	int Result = connect(Socket, &(sockaddr&)Addr, sizeof(Addr));
+	if (Result < 0)
+	{
+		close(Socket);
+		return 0;
+	}
+
+	return UPTRINT(Socket + 1);
+}
+
+////////////////////////////////////////////////////////////////////////////////
+UPTRINT TcpSocketListen(uint16 Port)
+{
+	int Socket = socket(AF_INET, SOCK_STREAM, IPPROTO_TCP);
+	if (Socket < 0)
+	{
+		return 0;
+	}
+
+	sockaddr_in SockAddr;
+	SockAddr.sin_family = AF_INET;
+	SockAddr.sin_addr.s_addr = 0;
+	SockAddr.sin_port = htons(Port);
+	int Result = bind(Socket, (sockaddr*)&SockAddr, sizeof(SockAddr));
+	if (Result < 0)
+	{
+		close(Socket);
+		return 0;
+	}
+
+	Result = listen(Socket, 1);
+	if (Result < 0)
+	{
+		close(Socket);
+		return 0;
+	}
+
+	return UPTRINT(Socket + 1);
+}
+
+////////////////////////////////////////////////////////////////////////////////
+UPTRINT TcpSocketAccept(UPTRINT Socket)
+{
+	int Inner = int(Socket - 1);
+
+	Inner = accept(Inner, nullptr, nullptr);
+	if (Inner < 0)
+	{
+		return 0;
+	}
+
+	return UPTRINT(Inner + 1);
+}
+
+////////////////////////////////////////////////////////////////////////////////
+void TcpSocketClose(UPTRINT Socket)
+{
+	int Inner = int(Socket - 1);
+	close(Inner);
+}
+
+////////////////////////////////////////////////////////////////////////////////
+bool TcpSocketSelect(UPTRINT Socket)
+{
+	int Inner = int(Socket - 1);
+	fd_set FdSet;
+	FD_SET(Inner, &FdSet);
+	timeval TimeVal = {};
+	return (select(0, &FdSet, nullptr, nullptr, &TimeVal) != 0);
+}
+
+////////////////////////////////////////////////////////////////////////////////
+bool TcpSocketSend(UPTRINT Socket, const void* Data, uint32 Size)
+{
+	int Inner = int(Socket - 1);
+	return send(Inner, (const char*)Data, Size, 0) == Size;
+}
+
+////////////////////////////////////////////////////////////////////////////////
+int32 TcpSocketRecv(UPTRINT Socket, void* Data, uint32 Size)
+{
+	int Inner = int(Socket - 1);
+	return recv(Inner, (char*)Data, Size, 0);
+}
+
+} // namespace Trace
