Removed the limiting of the number of buffers to check for exhaustion.

diff --git a/Engine/Source/Runtime/TraceLog/Private/Trace/Writer.cpp b/Engine/Source/Runtime/TraceLog/Private/Trace/Writer.cpp
--- a/Engine/Source/Runtime/TraceLog/Private/Trace/Writer.cpp
+++ b/Engine/Source/Runtime/TraceLog/Private/Trace/Writer.cpp
@@ -199,7 +199,7 @@ static void Writer_RetireBuffer(void (*DataSink)(const uint8*, uint32))
 ////////////////////////////////////////////////////////////////////////////////
 static void Writer_UpdateBuffers(void (*DataSink)(const uint8*, uint32))
 {
-	for (int i = 0; i < 3; ++i)
+	while (true)
 	{
 		FBuffer* Buffer = GActiveBuffer.load(std::memory_order_acquire);
 		if (GTailBuffer != Buffer)
@@ -219,6 +219,8 @@ static void Writer_UpdateBuffers(void (*DataSink)(const uint8*, uint32))
 			DataSink(Buffer->Data + GTailPreSent, Sendable);
 			GTailPreSent += Sendable;
 		}
+
+		break;
 	}
 }
 
