Increased the frequency of trace's worker thread

diff --git a/Engine/Source/Runtime/TraceLog/Private/Trace/Writer.cpp b/Engine/Source/Runtime/TraceLog/Private/Trace/Writer.cpp
--- a/Engine/Source/Runtime/TraceLog/Private/Trace/Writer.cpp
+++ b/Engine/Source/Runtime/TraceLog/Private/Trace/Writer.cpp
@@ -764,7 +764,7 @@ static void Writer_WorkerThread()
 {
 	while (!GWorkerThreadQuit)
 	{
-		const uint32 SleepMs = 100;
+		const uint32 SleepMs = 60;
 		ThreadSleep(SleepMs);
 
 		Writer_UpdateControl();
