Disable Trace's command-and-control server on Switch as it tries to start before networking/HTCS has been initialised.

diff --git a/Engine/Source/Runtime/TraceLog/Private/Trace/Writer.cpp b/Engine/Source/Runtime/TraceLog/Private/Trace/Writer.cpp
--- a/Engine/Source/Runtime/TraceLog/Private/Trace/Writer.cpp
+++ b/Engine/Source/Runtime/TraceLog/Private/Trace/Writer.cpp
@@ -719,6 +719,11 @@ static void Writer_UpdateControl()
 ////////////////////////////////////////////////////////////////////////////////
 static void Writer_InitializeControl()
 {
+#if PLATFORM_SWITCH
+	GControlState = EControlState::Failed;
+	return;
+#endif
+
 	Writer_ControlAddCommand("Connect", nullptr,
 		[] (void*, uint32 ArgC, ANSICHAR const* const* ArgV)
 		{
