Removed unused extern statement.

diff --git a/Engine/Source/Runtime/TraceLog/Private/Trace/Trace.cpp b/Engine/Source/Runtime/TraceLog/Private/Trace/Trace.cpp
--- a/Engine/Source/Runtime/TraceLog/Private/Trace/Trace.cpp
+++ b/Engine/Source/Runtime/TraceLog/Private/Trace/Trace.cpp
@@ -12,7 +12,6 @@ namespace Private
 {
 
 ////////////////////////////////////////////////////////////////////////////////
-extern UPTRINT	GPendingDataHandle;
 void			Writer_Flush();
 bool			Writer_Connect(const ANSICHAR*);
 uint32			Writer_EventToggle(const ANSICHAR*, bool);
