Changed the listen-accept API to continue to be poll-based but not reliant on select calls. This scheme maps better to how Switches two socket subsystems work.

diff --git a/Engine/Source/Runtime/TraceLog/Private/Trace/Detail/PosixTrace.inl b/Engine/Source/Runtime/TraceLog/Private/Trace/Detail/PosixTrace.inl
--- a/Engine/Source/Runtime/TraceLog/Private/Trace/Detail/PosixTrace.inl
+++ b/Engine/Source/Runtime/TraceLog/Private/Trace/Detail/PosixTrace.inl
@@ -1,10 +1,11 @@
 // Copyright 1998-2019 Epic Games, Inc. All Rights Reserved.
 
-#include <pthread.h>
-#include <unistd.h>
-
 #include <arpa/inet.h>
+#include <errno.h>
+#include <fcntl.h>
+#include <pthread.h>
 #include <sys/socket.h>
+#include <unistd.h>
 
 namespace Trace
 {
@@ -107,21 +108,41 @@ UPTRINT TcpSocketListen(uint16 Port)
 		return 0;
 	}
 
+	int Flags = fcntl(Socket, F_GETFL, 0);
+	if (Flags < 0 || !fcntl(Socket, F_SETFL, Flags|O_NONBLOCK))
+	{
+		close(Socket);
+		return 0;
+	}
+
 	return UPTRINT(Socket + 1);
 }
 
 ////////////////////////////////////////////////////////////////////////////////
-UPTRINT TcpSocketAccept(UPTRINT Socket)
+int TcpSocketAccept(UPTRINT Socket, UPTRINT& Out)
 {
 	int Inner = int(Socket - 1);
 
 	Inner = accept(Inner, nullptr, nullptr);
 	if (Inner < 0)
 	{
+		if (Inner == EWOULDBLOCK || Inner == EAGAIN)
+		{
+			return 0;
+		}
+
+		return -1;
+	}
+
+	int Flags = fcntl(Inner, F_GETFL, 0);
+	if (Flags < 0 || !fcntl(Inner, F_SETFL, Flags & ~O_NONBLOCK))
+	{
+		close(Socket);
 		return 0;
 	}
 
-	return UPTRINT(Inner + 1);
+	Out = UPTRINT(Inner + 1);
+	return 1;
 }
 
 ////////////////////////////////////////////////////////////////////////////////
diff --git a/Engine/Source/Runtime/TraceLog/Private/Trace/Detail/Windows/WindowsTrace.cpp b/Engine/Source/Runtime/TraceLog/Private/Trace/Detail/Windows/WindowsTrace.cpp
--- a/Engine/Source/Runtime/TraceLog/Private/Trace/Detail/Windows/WindowsTrace.cpp
+++ b/Engine/Source/Runtime/TraceLog/Private/Trace/Detail/Windows/WindowsTrace.cpp
@@ -190,21 +190,36 @@ UPTRINT TcpSocketListen(uint16 Port)
 		return 0;
 	}
 
+	unsigned long NonBlockingMode = 1;
+	if (ioctlsocket(Socket, FIONBIO, &NonBlockingMode) == SOCKET_ERROR)
+	{
+		closesocket(Socket);
+		return 0;
+	}
+
 	return UPTRINT(Socket + 1);
 }
 
 ////////////////////////////////////////////////////////////////////////////////
-UPTRINT TcpSocketAccept(UPTRINT Socket)
+int32 TcpSocketAccept(UPTRINT Socket, UPTRINT& Out)
 {
 	SOCKET Inner = Socket - 1;
 
 	Inner = accept(Inner, nullptr, nullptr);
 	if (Inner == INVALID_SOCKET)
 	{
+		return (WSAGetLastError() == WSAEWOULDBLOCK) - 1;
+	}
+
+	unsigned long NonBlockingMode = 0;
+	if (ioctlsocket(Socket, FIONBIO, &NonBlockingMode) == SOCKET_ERROR)
+	{
+		closesocket(Socket);
 		return 0;
 	}
 
-	return UPTRINT(Inner + 1);
+	Out = UPTRINT(Inner + 1);
+	return 1;
 }
 
 ////////////////////////////////////////////////////////////////////////////////
diff --git a/Engine/Source/Runtime/TraceLog/Private/Trace/Writer.cpp b/Engine/Source/Runtime/TraceLog/Private/Trace/Writer.cpp
--- a/Engine/Source/Runtime/TraceLog/Private/Trace/Writer.cpp
+++ b/Engine/Source/Runtime/TraceLog/Private/Trace/Writer.cpp
@@ -20,7 +20,7 @@ void	MemoryFree(void*, SIZE_T);
 void	MemoryMap(void*, SIZE_T);
 UPTRINT	TcpSocketConnect(const ANSICHAR*, uint16);
 UPTRINT	TcpSocketListen(uint16);
-UPTRINT	TcpSocketAccept(UPTRINT);
+int32	TcpSocketAccept(UPTRINT, UPTRINT&);
 void	TcpSocketClose(UPTRINT);
 bool	TcpSocketSelect(UPTRINT);
 int32	TcpSocketRecv(UPTRINT, void*, uint32);
@@ -582,14 +582,16 @@ static bool Writer_ControlListen()
 ////////////////////////////////////////////////////////////////////////////////
 static bool Writer_ControlAccept()
 {
-	if (!TcpSocketSelect(GControlListen))
-	{
-		return false;
-	}
-
-	UPTRINT Socket = TcpSocketAccept(GControlListen);
-	if (!Socket)
+	UPTRINT Socket;
+	int Return = TcpSocketAccept(GControlListen, Socket);
+	if (Return <= 0)
 	{
+		if (Return == -1)
+		{
+			TcpSocketClose(GControlListen);
+			GControlListen = 0;
+			GControlState = EControlState::Failed;
+		}
 		return false;
 	}
 
