Renamed Select to HasData to avoid confusion for those familiar with select().

diff --git a/Engine/Source/Runtime/TraceLog/Private/Trace/Detail/PosixTrace.inl b/Engine/Source/Runtime/TraceLog/Private/Trace/Detail/PosixTrace.inl
--- a/Engine/Source/Runtime/TraceLog/Private/Trace/Detail/PosixTrace.inl
+++ b/Engine/Source/Runtime/TraceLog/Private/Trace/Detail/PosixTrace.inl
@@ -153,7 +153,7 @@ void TcpSocketClose(UPTRINT Socket)
 }
 
 ////////////////////////////////////////////////////////////////////////////////
-bool TcpSocketSelect(UPTRINT Socket)
+bool TcpSocketHasData(UPTRINT Socket)
 {
 	int Inner = int(Socket - 1);
 	fd_set FdSet;
diff --git a/Engine/Source/Runtime/TraceLog/Private/Trace/Detail/Windows/WindowsTrace.cpp b/Engine/Source/Runtime/TraceLog/Private/Trace/Detail/Windows/WindowsTrace.cpp
--- a/Engine/Source/Runtime/TraceLog/Private/Trace/Detail/Windows/WindowsTrace.cpp
+++ b/Engine/Source/Runtime/TraceLog/Private/Trace/Detail/Windows/WindowsTrace.cpp
@@ -230,7 +230,7 @@ void TcpSocketClose(UPTRINT Socket)
 }
 
 ////////////////////////////////////////////////////////////////////////////////
-bool TcpSocketSelect(UPTRINT Socket)
+bool TcpSocketHasData(UPTRINT Socket)
 {
 	SOCKET Inner = Socket - 1;
 	fd_set FdSet = { 1, { Inner }, };
diff --git a/Engine/Source/Runtime/TraceLog/Private/Trace/Writer.cpp b/Engine/Source/Runtime/TraceLog/Private/Trace/Writer.cpp
--- a/Engine/Source/Runtime/TraceLog/Private/Trace/Writer.cpp
+++ b/Engine/Source/Runtime/TraceLog/Private/Trace/Writer.cpp
@@ -22,7 +22,7 @@ UPTRINT	TcpSocketConnect(const ANSICHAR*, uint16);
 UPTRINT	TcpSocketListen(uint16);
 int32	TcpSocketAccept(UPTRINT, UPTRINT&);
 void	TcpSocketClose(UPTRINT);
-bool	TcpSocketSelect(UPTRINT);
+bool	TcpSocketHasData(UPTRINT);
 int32	TcpSocketRecv(UPTRINT, void*, uint32);
 bool	TcpSocketSend(UPTRINT, const void*, uint32);
 UPTRINT	ThreadCreate(const ANSICHAR*, void (*)());
@@ -608,7 +608,7 @@ static void Writer_ControlRecv()
 
 	ANSICHAR Buffer[512];
 	ANSICHAR* __restrict Head = Buffer;
-	while (TcpSocketSelect(GControlSocket))
+	while (TcpSocketHasData(GControlSocket))
 	{
 		int32 ReadSize = int32(UPTRINT(Buffer + sizeof(Buffer) - Head));
 		int32 Recvd = TcpSocketRecv(GControlSocket, Head, ReadSize);
