Changed the Tcp-prefix send/recv/close API to be a generic Io-prefixed API so that single paths can be agnostic of different types of IO primitives (sockets, pipes, files, the Windows console, etc.)

diff --git a/Engine/Source/Runtime/TraceLog/Private/Trace/Detail/PosixTrace.inl b/Engine/Source/Runtime/TraceLog/Private/Trace/Detail/PosixTrace.inl
--- a/Engine/Source/Runtime/TraceLog/Private/Trace/Detail/PosixTrace.inl
+++ b/Engine/Source/Runtime/TraceLog/Private/Trace/Detail/PosixTrace.inl
@@ -145,13 +145,6 @@ int TcpSocketAccept(UPTRINT Socket, UPTRINT& Out)
 	return 1;
 }
 
-////////////////////////////////////////////////////////////////////////////////
-void TcpSocketClose(UPTRINT Socket)
-{
-	int Inner = int(Socket - 1);
-	close(Inner);
-}
-
 ////////////////////////////////////////////////////////////////////////////////
 bool TcpSocketHasData(UPTRINT Socket)
 {
@@ -162,18 +155,27 @@ bool TcpSocketHasData(UPTRINT Socket)
 	return (select(0, &FdSet, nullptr, nullptr, &TimeVal) != 0);
 }
 
+
+
 ////////////////////////////////////////////////////////////////////////////////
-bool TcpSocketSend(UPTRINT Socket, const void* Data, uint32 Size)
+bool IoWrite(UPTRINT Handle, const void* Data, uint32 Size)
 {
-	int Inner = int(Socket - 1);
-	return send(Inner, (const char*)Data, Size, 0) == Size;
+	int Inner = int(Handle - 1);
+	return write(Inner, Data, Size) == Size;
 }
 
 ////////////////////////////////////////////////////////////////////////////////
-int32 TcpSocketRecv(UPTRINT Socket, void* Data, uint32 Size)
+int32 IoRead(UPTRINT Handle, void* Data, uint32 Size)
 {
-	int Inner = int(Socket - 1);
-	return recv(Inner, (char*)Data, Size, 0);
+	int Inner = int(Handle - 1);
+	return read(Inner, Data, Size);
+}
+
+////////////////////////////////////////////////////////////////////////////////
+void IoClose(UPTRINT Handle)
+{
+	int Inner = int(Handle - 1);
+	close(Inner);
 }
 
 } // namespace Trace
diff --git a/Engine/Source/Runtime/TraceLog/Private/Trace/Detail/Windows/WindowsTrace.cpp b/Engine/Source/Runtime/TraceLog/Private/Trace/Detail/Windows/WindowsTrace.cpp
--- a/Engine/Source/Runtime/TraceLog/Private/Trace/Detail/Windows/WindowsTrace.cpp
+++ b/Engine/Source/Runtime/TraceLog/Private/Trace/Detail/Windows/WindowsTrace.cpp
@@ -222,13 +222,6 @@ int32 TcpSocketAccept(UPTRINT Socket, UPTRINT& Out)
 	return 1;
 }
 
-////////////////////////////////////////////////////////////////////////////////
-void TcpSocketClose(UPTRINT Socket)
-{
-	SOCKET Inner = Socket - 1;
-	closesocket(Inner);
-}
-
 ////////////////////////////////////////////////////////////////////////////////
 bool TcpSocketHasData(UPTRINT Socket)
 {
@@ -238,18 +231,41 @@ bool TcpSocketHasData(UPTRINT Socket)
 	return (select(0, &FdSet, nullptr, nullptr, &TimeVal) != 0);
 }
 
+
+
 ////////////////////////////////////////////////////////////////////////////////
-bool TcpSocketSend(UPTRINT Socket, const void* Data, uint32 Size)
+bool IoWrite(UPTRINT Handle, const void* Data, uint32 Size)
 {
-	SOCKET Inner = Socket - 1;
-	return send(Inner, (const char*)Data, Size, 0) == Size;
+	HANDLE Inner = HANDLE(Handle - 1);
+
+	DWORD BytesWritten = 0;
+	if (!WriteFile(Inner, (const char*)Data, Size, &BytesWritten, nullptr))
+	{
+		return false;
+	}
+
+	return (BytesWritten == Size);
 }
 
 ////////////////////////////////////////////////////////////////////////////////
-int32 TcpSocketRecv(UPTRINT Socket, void* Data, uint32 Size)
+int32 IoRead(UPTRINT Handle, void* Data, uint32 Size)
 {
-	SOCKET Inner = Socket - 1;
-	return recv(Inner, (char*)Data, Size, 0);
+	HANDLE Inner = HANDLE(Handle - 1);
+
+	DWORD BytesRead = 0;
+	if (!ReadFile(Inner, (char*)Data, Size, &BytesRead, nullptr))
+	{
+		return -1;
+	}
+
+	return BytesRead;
+}
+
+////////////////////////////////////////////////////////////////////////////////
+void IoClose(UPTRINT Handle)
+{
+	HANDLE Inner = HANDLE(Handle - 1);
+	CloseHandle(Inner);
 }
 
 } // namespace Trace
diff --git a/Engine/Source/Runtime/TraceLog/Private/Trace/Writer.cpp b/Engine/Source/Runtime/TraceLog/Private/Trace/Writer.cpp
--- a/Engine/Source/Runtime/TraceLog/Private/Trace/Writer.cpp
+++ b/Engine/Source/Runtime/TraceLog/Private/Trace/Writer.cpp
@@ -21,10 +21,10 @@ void	MemoryMap(void*, SIZE_T);
 UPTRINT	TcpSocketConnect(const ANSICHAR*, uint16);
 UPTRINT	TcpSocketListen(uint16);
 int32	TcpSocketAccept(UPTRINT, UPTRINT&);
-void	TcpSocketClose(UPTRINT);
 bool	TcpSocketHasData(UPTRINT);
-int32	TcpSocketRecv(UPTRINT, void*, uint32);
-bool	TcpSocketSend(UPTRINT, const void*, uint32);
+int32	IoRead(UPTRINT, void*, uint32);
+bool	IoWrite(UPTRINT, const void*, uint32);
+void	IoClose(UPTRINT);
 UPTRINT	ThreadCreate(const ANSICHAR*, void (*)());
 uint32	ThreadGetCurrentId();
 void	ThreadSleep(uint32 Milliseconds);
@@ -411,7 +411,7 @@ static void Writer_UpdateData()
 	{
 		if (GDataHandle)
 		{
-			TcpSocketClose(GPendingDataHandle);
+			IoClose(GPendingDataHandle);
 			GPendingDataHandle = 0;
 		}
 		else
@@ -421,17 +421,17 @@ static void Writer_UpdateData()
 
 			// Handshake.
 			const uint32 Magic = 'TRCE';
-			bool bOk = TcpSocketSend(GDataHandle, &Magic, sizeof(Magic));
+			bool bOk = IoWrite(GDataHandle, &Magic, sizeof(Magic));
 
 			// Stream header
 			const struct {
 				uint8 Format;
 				uint8 Parameter;
 			} TransportHeader = { 1 };
-			bOk &= TcpSocketSend(GDataHandle, &TransportHeader, sizeof(TransportHeader));
+			bOk &= IoWrite(GDataHandle, &TransportHeader, sizeof(TransportHeader));
 
 			// Passively collected data
-			bOk &= TcpSocketSend(GDataHandle, GHoldBuffer.GetData(), GHoldBuffer.GetSize());
+			bOk &= IoWrite(GDataHandle, GHoldBuffer.GetData(), GHoldBuffer.GetSize());
 
 			if (bOk)
 			{
@@ -440,7 +440,7 @@ static void Writer_UpdateData()
 			}
 			else
 			{
-				TcpSocketClose(GDataHandle);
+				IoClose(GDataHandle);
 				GDataHandle = 0;
 			}
 		}
@@ -454,9 +454,9 @@ static void Writer_UpdateData()
 		{
 			Writer_UpdateBuffers([] (const uint8* Data, uint32 Size)
 			{
-				if (GDataHandle && !TcpSocketSend(GDataHandle, Data, Size))
+				if (GDataHandle && !IoWrite(GDataHandle, Data, Size))
 				{
-					TcpSocketClose(GDataHandle);
+					IoClose(GDataHandle);
 					GDataHandle = 0;
 				}
 			});
@@ -588,7 +588,7 @@ static bool Writer_ControlAccept()
 	{
 		if (Return == -1)
 		{
-			TcpSocketClose(GControlListen);
+			IoClose(GControlListen);
 			GControlListen = 0;
 			GControlState = EControlState::Failed;
 		}
@@ -611,10 +611,10 @@ static void Writer_ControlRecv()
 	while (TcpSocketHasData(GControlSocket))
 	{
 		int32 ReadSize = int32(UPTRINT(Buffer + sizeof(Buffer) - Head));
-		int32 Recvd = TcpSocketRecv(GControlSocket, Head, ReadSize);
+		int32 Recvd = IoRead(GControlSocket, Head, ReadSize);
 		if (Recvd <= 0)
 		{
-			TcpSocketClose(GControlSocket);
+			IoClose(GControlSocket);
 			GControlSocket = 0;
 			GControlState = EControlState::Listening;
 			break;
@@ -755,7 +755,7 @@ static void Writer_ShutdownControl()
 {
 	if (GControlListen)
 	{
-		TcpSocketClose(GControlListen);
+		IoClose(GControlListen);
 		GControlListen = 0;
 	}
 }
