Correct some casts/arithmetic to better match intention.

diff --git a/Engine/Source/Runtime/TraceLog/Private/Trace/Detail/PosixTrace.inl b/Engine/Source/Runtime/TraceLog/Private/Trace/Detail/PosixTrace.inl
--- a/Engine/Source/Runtime/TraceLog/Private/Trace/Detail/PosixTrace.inl
+++ b/Engine/Source/Runtime/TraceLog/Private/Trace/Detail/PosixTrace.inl
@@ -78,7 +78,7 @@ UPTRINT TcpSocketConnect(const ANSICHAR* Host, uint16 Port)
 		return 0;
 	}
 
-	return UPTRINT(Socket + 1);
+	return UPTRINT(Socket) + 1;
 }
 
 ////////////////////////////////////////////////////////////////////////////////
@@ -115,7 +115,7 @@ UPTRINT TcpSocketListen(uint16 Port)
 		return 0;
 	}
 
-	return UPTRINT(Socket + 1);
+	return UPTRINT(Socket) + 1;
 }
 
 ////////////////////////////////////////////////////////////////////////////////
@@ -141,7 +141,7 @@ int TcpSocketAccept(UPTRINT Socket, UPTRINT& Out)
 		return 0;
 	}
 
-	Out = UPTRINT(Inner + 1);
+	Out = UPTRINT(Inner) + 1;
 	return 1;
 }
 
diff --git a/Engine/Source/Runtime/TraceLog/Private/Trace/Detail/Windows/WindowsTrace.cpp b/Engine/Source/Runtime/TraceLog/Private/Trace/Detail/Windows/WindowsTrace.cpp
--- a/Engine/Source/Runtime/TraceLog/Private/Trace/Detail/Windows/WindowsTrace.cpp
+++ b/Engine/Source/Runtime/TraceLog/Private/Trace/Detail/Windows/WindowsTrace.cpp
@@ -158,7 +158,7 @@ UPTRINT TcpSocketConnect(const ANSICHAR* Host, uint16 Port)
 		return 0;
 	}
 
-	return UPTRINT(Socket + 1);
+	return UPTRINT(Socket) + 1;
 }
 
 ////////////////////////////////////////////////////////////////////////////////
@@ -197,7 +197,7 @@ UPTRINT TcpSocketListen(uint16 Port)
 		return 0;
 	}
 
-	return UPTRINT(Socket + 1);
+	return UPTRINT(Socket) + 1;
 }
 
 ////////////////////////////////////////////////////////////////////////////////
@@ -218,7 +218,7 @@ int32 TcpSocketAccept(UPTRINT Socket, UPTRINT& Out)
 		return 0;
 	}
 
-	Out = UPTRINT(Inner + 1);
+	Out = UPTRINT(Inner) + 1;
 	return 1;
 }
 
