Support for tracing to a file.

diff --git a/Engine/Source/Runtime/TraceLog/Private/Trace/Detail/PosixTrace.inl b/Engine/Source/Runtime/TraceLog/Private/Trace/Detail/PosixTrace.inl
--- a/Engine/Source/Runtime/TraceLog/Private/Trace/Detail/PosixTrace.inl
+++ b/Engine/Source/Runtime/TraceLog/Private/Trace/Detail/PosixTrace.inl
@@ -178,4 +178,24 @@ void IoClose(UPTRINT Handle)
 	close(Inner);
 }
 
+
+
+////////////////////////////////////////////////////////////////////////////////
+UPTRINT FileOpen(const ANSICHAR* Path, const ANSICHAR* Mode)
+{
+	int Flags = (*Mode == 'a' || *Mode == 'w') ? O_RDWR : O_RDONLY;
+	int Out = open(Path, Flags);
+	if (Out < 0)
+	{
+		return 0;
+	}
+
+	if (*Mode == '0')
+	{
+		lseek(Out, 0, SEEK_END);
+	}
+
+	return UPTRINT(Out) + 1;
+}
+
 } // namespace Trace
diff --git a/Engine/Source/Runtime/TraceLog/Private/Trace/Detail/Windows/WindowsTrace.cpp b/Engine/Source/Runtime/TraceLog/Private/Trace/Detail/Windows/WindowsTrace.cpp
--- a/Engine/Source/Runtime/TraceLog/Private/Trace/Detail/Windows/WindowsTrace.cpp
+++ b/Engine/Source/Runtime/TraceLog/Private/Trace/Detail/Windows/WindowsTrace.cpp
@@ -145,7 +145,10 @@ UPTRINT TcpSocketConnect(const ANSICHAR* Host, uint16 Port)
 	auto* SockAddr = (sockaddr_in*)Info->ai_addr;
 	SockAddr->sin_port = htons(Port);
 
-	SOCKET Socket = socket(AF_INET, SOCK_STREAM, IPPROTO_TCP);
+	// socket() will create a socket with overlapped IO support which we don't
+	// want as it complicates sharing Io*() API with FileOpen(). So we use
+	// WSASocket instead which affords us more control over socket properties
+	SOCKET Socket = WSASocket(AF_INET, SOCK_STREAM, IPPROTO_TCP, nullptr, 0, 0);
 	if (Socket == INVALID_SOCKET)
 	{
 		return 0;
@@ -166,7 +169,8 @@ UPTRINT TcpSocketListen(uint16 Port)
 {
 	TcpSocketInitialize();
 
-	SOCKET Socket = socket(AF_INET, SOCK_STREAM, IPPROTO_TCP);
+	// See TcpSocketConnect() for why WSASocket() is used here.
+	SOCKET Socket = WSASocket(AF_INET, SOCK_STREAM, IPPROTO_TCP, nullptr, 0, 0);
 	if (Socket == INVALID_SOCKET)
 	{
 		return 0;
@@ -268,6 +272,37 @@ void IoClose(UPTRINT Handle)
 	CloseHandle(Inner);
 }
 
+
+
+////////////////////////////////////////////////////////////////////////////////
+UPTRINT FileOpen(const ANSICHAR* Path, const ANSICHAR* Mode)
+{
+	DWORD Access = GENERIC_READ;
+	DWORD Disposition = OPEN_EXISTING;
+	switch (*Mode)
+	{
+	case 'a':
+	case 'w':
+		Access |= GENERIC_WRITE;
+		Disposition = (*Mode == 'a') ? OPEN_ALWAYS : CREATE_ALWAYS;
+		break;
+	}
+
+	DWORD Flags = FILE_ATTRIBUTE_NORMAL;
+	HANDLE Out = CreateFileA(Path, Access, FILE_SHARE_READ|FILE_SHARE_WRITE, nullptr, Disposition, Flags, nullptr);
+	if (Out == INVALID_HANDLE_VALUE)
+	{
+		return 0;
+	}
+
+	if (*Mode == 'a')
+	{
+		SetFilePointer(Out, 0, nullptr, FILE_END);
+	}
+
+	return UPTRINT(Out) + 1;
+}
+
 } // namespace Trace
 
 #pragma warning(pop)
diff --git a/Engine/Source/Runtime/TraceLog/Private/Trace/Trace.cpp b/Engine/Source/Runtime/TraceLog/Private/Trace/Trace.cpp
--- a/Engine/Source/Runtime/TraceLog/Private/Trace/Trace.cpp
+++ b/Engine/Source/Runtime/TraceLog/Private/Trace/Trace.cpp
@@ -12,9 +12,10 @@ namespace Private
 {
 
 ////////////////////////////////////////////////////////////////////////////////
-void			Writer_Flush();
-bool			Writer_Connect(const ANSICHAR*);
-uint32			Writer_EventToggle(const ANSICHAR*, bool);
+void	Writer_Flush();
+bool	Writer_Connect(const ANSICHAR*);
+bool	Writer_Open(const ANSICHAR*);
+uint32	Writer_EventToggle(const ANSICHAR*, bool);
 
 } // namespace Private
 
@@ -43,6 +44,14 @@ bool Connect(const TCHAR* InHost)
 	return Private::Writer_Connect(Host);
 }
 
+////////////////////////////////////////////////////////////////////////////////
+bool Open(const TCHAR* InPath)
+{
+	char Path[512];
+	ToAnsiCheap(Path, InPath);
+	return Private::Writer_Open(Path);
+}
+
 ////////////////////////////////////////////////////////////////////////////////
 uint32 ToggleEvent(const TCHAR* Wildcard, bool bState)
 {
diff --git a/Engine/Source/Runtime/TraceLog/Private/Trace/Writer.cpp b/Engine/Source/Runtime/TraceLog/Private/Trace/Writer.cpp
--- a/Engine/Source/Runtime/TraceLog/Private/Trace/Writer.cpp
+++ b/Engine/Source/Runtime/TraceLog/Private/Trace/Writer.cpp
@@ -25,6 +25,7 @@ bool	TcpSocketHasData(UPTRINT);
 int32	IoRead(UPTRINT, void*, uint32);
 bool	IoWrite(UPTRINT, const void*, uint32);
 void	IoClose(UPTRINT);
+UPTRINT	FileOpen(const ANSICHAR*, const ANSICHAR*);
 UPTRINT	ThreadCreate(const ANSICHAR*, void (*)());
 uint32	ThreadGetCurrentId();
 void	ThreadSleep(uint32 Milliseconds);
@@ -508,6 +509,7 @@ struct FControlCommands
 
 ////////////////////////////////////////////////////////////////////////////////
 bool	Writer_Connect(const ANSICHAR*);
+bool	Writer_Open(const ANSICHAR*);
 uint32	Writer_EventToggle(const ANSICHAR*, bool);
 
 ////////////////////////////////////////////////////////////////////////////////
@@ -736,6 +738,16 @@ static void Writer_InitializeControl()
 		}
 	);
 
+	Writer_ControlAddCommand("Open", nullptr,
+		[] (void*, uint32 ArgC, ANSICHAR const* const* ArgV)
+		{
+			if (ArgC > 0)
+			{
+				Writer_Open(ArgV[0]);
+			}
+		}
+	);
+
 	Writer_ControlAddCommand("ToggleEvent", nullptr,
 		[] (void*, uint32 ArgC, ANSICHAR const* const* ArgV)
 		{
@@ -880,6 +892,21 @@ bool Writer_Connect(const ANSICHAR* Host)
 	return true;
 }
 
+////////////////////////////////////////////////////////////////////////////////
+bool Writer_Open(const ANSICHAR* Path)
+{
+	Writer_Initialize();
+
+	UPTRINT DataHandle = FileOpen(Path, "w");
+	if (!DataHandle)
+	{
+		return false;
+	}
+
+	GPendingDataHandle = DataHandle;
+	return true;
+}
+
 
 
 ////////////////////////////////////////////////////////////////////////////////
