Issue with Centos7 libc.so.6 GLIBC_2.17 not having __call_tls_dtors implemented which ends up cleaning the TLS after shutdown functions are called. Leaving dangling pointers around

diff --git a/Engine/Source/Runtime/TraceLog/Private/Trace/Writer.cpp b/Engine/Source/Runtime/TraceLog/Private/Trace/Writer.cpp
--- a/Engine/Source/Runtime/TraceLog/Private/Trace/Writer.cpp
+++ b/Engine/Source/Runtime/TraceLog/Private/Trace/Writer.cpp
@@ -77,6 +77,8 @@ void Writer_InitializeTiming()
 }
 
 
+////////////////////////////////////////////////////////////////////////////////
+static bool GInitialized = false;
 
 ////////////////////////////////////////////////////////////////////////////////
 thread_local					FWriteTlsContext TlsContext;
@@ -114,7 +116,7 @@ FWriteTlsContext::FWriteTlsContext()
 ////////////////////////////////////////////////////////////////////////////////
 FWriteTlsContext::~FWriteTlsContext()
 {
-	if (HasValidBuffer())
+	if (GInitialized && HasValidBuffer())
 	{
 		UPTRINT EtxOffset = ~UPTRINT((uint8*)Buffer - Buffer->Cursor);
 		AtomicStoreRelaxed(&(Buffer->EtxOffset), EtxOffset);
@@ -972,10 +974,6 @@ static void Writer_WorkerThread()
 }
 
 
-
-////////////////////////////////////////////////////////////////////////////////
-static bool GInitialized = false;
-
 ////////////////////////////////////////////////////////////////////////////////
 static void Writer_LogHeader()
 {
