Fixed Engine compilation for Android platform: added required #include (Contributed by [REDACTED_PERSON])

#PR #6511
#ue4
#android

diff --git a/Engine/Source/Runtime/TraceLog/Private/Trace/Detail/Android/AndroidTrace.cpp b/Engine/Source/Runtime/TraceLog/Private/Trace/Detail/Android/AndroidTrace.cpp
--- a/Engine/Source/Runtime/TraceLog/Private/Trace/Detail/Android/AndroidTrace.cpp
+++ b/Engine/Source/Runtime/TraceLog/Private/Trace/Detail/Android/AndroidTrace.cpp
@@ -5,6 +5,7 @@
 #if UE_TRACE_ENABLED
 
 #include <arpa/inet.h>
+#include <errno.h>
 #include <fcntl.h>
 #include <pthread.h>
 #include <sys/mman.h>
