Fix Linux Arm64 build breakage

(ushell-p4-cherrypick of 8498739 by michael.sartain)

diff --git a/Engine/Source/Runtime/TraceLog/Private/Trace/Writer.cpp b/Engine/Source/Runtime/TraceLog/Private/Trace/Writer.cpp
--- a/Engine/Source/Runtime/TraceLog/Private/Trace/Writer.cpp
+++ b/Engine/Source/Runtime/TraceLog/Private/Trace/Writer.cpp
@@ -9,6 +9,11 @@
 
 #if PLATFORM_CPU_X86_FAMILY
 	#include <emmintrin.h>
+	#define PLATFORM_YIELD()	_mm_pause()
+#elif PLATFORM_CPU_ARM_FAMILY
+	#define PLATFORM_YIELD()	__builtin_arm_yield();
+#else
+	#error Unsupported architecture!
 #endif
 
 namespace Trace
@@ -42,13 +47,7 @@ namespace Private
 ////////////////////////////////////////////////////////////////////////////////
 inline void Writer_Yield()
 {
-#if PLATFORM_CPU_X86_FAMILY
-	_mm_pause();
-#elif PLATFORM_CPU_ARM_FAMILY
-	__builtin_arm_yield();
-#else
-	#error Unsupported platform!
-#endif
+	PLATFORM_YIELD();
 }
 
 
