Moved TraceLog's inner details into its Private namespace

diff --git a/Engine/Source/Runtime/TraceLog/Private/Trace/Detail/Linux/LinuxTrace.cpp b/Engine/Source/Runtime/TraceLog/Private/Trace/Detail/Linux/LinuxTrace.cpp
--- a/Engine/Source/Runtime/TraceLog/Private/Trace/Detail/Linux/LinuxTrace.cpp
+++ b/Engine/Source/Runtime/TraceLog/Private/Trace/Detail/Linux/LinuxTrace.cpp
@@ -9,8 +9,8 @@
 #include <sys/mman.h>
 #include <time.h>
 
-namespace Trace
-{
+namespace Trace {
+namespace Private {
 
 ////////////////////////////////////////////////////////////////////////////////
 uint8* MemoryReserve(SIZE_T Size)
@@ -52,6 +52,7 @@ uint64 TimeGetTimestamp()
 	return (uint64(Time.tv_sec) * (1000 * 1000 * 1000)) + Time.tv_nsec;
 }
 
+} // namespace Private
 } // namespace Trace
 
 #endif // UE_TRACE_ENABLED
diff --git a/Engine/Source/Runtime/TraceLog/Private/Trace/Detail/PosixTrace.inl b/Engine/Source/Runtime/TraceLog/Private/Trace/Detail/PosixTrace.inl
--- a/Engine/Source/Runtime/TraceLog/Private/Trace/Detail/PosixTrace.inl
+++ b/Engine/Source/Runtime/TraceLog/Private/Trace/Detail/PosixTrace.inl
@@ -7,8 +7,8 @@
 #include <sys/socket.h>
 #include <unistd.h>
 
-namespace Trace
-{
+namespace Trace {
+namespace Private {
 
 ////////////////////////////////////////////////////////////////////////////////
 UPTRINT ThreadCreate(const ANSICHAR* Name, void (*Entry)())
@@ -198,4 +198,5 @@ UPTRINT FileOpen(const ANSICHAR* Path, const ANSICHAR* Mode)
 	return UPTRINT(Out) + 1;
 }
 
+} // namespace Private
 } // namespace Trace
diff --git a/Engine/Source/Runtime/TraceLog/Private/Trace/Detail/Windows/WindowsTrace.cpp b/Engine/Source/Runtime/TraceLog/Private/Trace/Detail/Windows/WindowsTrace.cpp
--- a/Engine/Source/Runtime/TraceLog/Private/Trace/Detail/Windows/WindowsTrace.cpp
+++ b/Engine/Source/Runtime/TraceLog/Private/Trace/Detail/Windows/WindowsTrace.cpp
@@ -16,8 +16,8 @@
 #pragma warning(disable : 6250) // VirtualFree() missing MEM_RELEASE - a false positive
 #pragma warning(disable : 6031) // WSAStartup() return ignore  - we're error tolerant
 
-namespace Trace
-{
+namespace Trace {
+namespace Private {
 
 ////////////////////////////////////////////////////////////////////////////////
 uint8* MemoryReserve(SIZE_T Size)
@@ -303,6 +303,7 @@ UPTRINT FileOpen(const ANSICHAR* Path, const ANSICHAR* Mode)
 	return UPTRINT(Out) + 1;
 }
 
+} // namespace Private
 } // namespace Trace
 
 #pragma warning(pop)
diff --git a/Engine/Source/Runtime/TraceLog/Private/Trace/EventDef.cpp b/Engine/Source/Runtime/TraceLog/Private/Trace/EventDef.cpp
--- a/Engine/Source/Runtime/TraceLog/Private/Trace/EventDef.cpp
+++ b/Engine/Source/Runtime/TraceLog/Private/Trace/EventDef.cpp
@@ -10,7 +10,6 @@
 namespace Trace
 {
 
-////////////////////////////////////////////////////////////////////////////////
 namespace Private
 {
 
