Added Private namespace to TraceLog's Unix detail too.

diff --git a/Engine/Source/Runtime/TraceLog/Private/Trace/Unix/Detail.cpp b/Engine/Source/Runtime/TraceLog/Private/Trace/Unix/Detail.cpp
--- a/Engine/Source/Runtime/TraceLog/Private/Trace/Unix/Detail.cpp
+++ b/Engine/Source/Runtime/TraceLog/Private/Trace/Unix/Detail.cpp
@@ -15,8 +15,8 @@
 	#include <sys/syscall.h>
 #endif // _GNU_SOURCE
 
-namespace Trace
-{
+namespace Trace {
+namespace Private {
 
 ////////////////////////////////////////////////////////////////////////////////
 uint8* MemoryReserve(SIZE_T Size)
@@ -229,6 +229,7 @@ int32 TcpSocketRecv(UPTRINT Socket, void* Data, uint32 Size)
 	return static_cast<int32>(recv(Inner, (char*)Data, Size, 0));
 }
 
+} // namespace Private
 } // namespace Trace
 
 #endif // UE_TRACE_ENABLED
