Corrected a single-line-if rule break.

diff --git a/Engine/Source/Runtime/TraceLog/Public/Trace/Private/EventDef.inl b/Engine/Source/Runtime/TraceLog/Public/Trace/Private/EventDef.inl
--- a/Engine/Source/Runtime/TraceLog/Public/Trace/Private/EventDef.inl
+++ b/Engine/Source/Runtime/TraceLog/Public/Trace/Private/EventDef.inl
@@ -34,7 +34,9 @@ inline FEventDef::FLogScope::FLogScope(uint16 EventUid, uint16 Size, uint8* Out)
 inline FEventDef::FLogScope::~FLogScope()
 {
 	if (!bOutOfBand)
+	{
 		Writer_EndLog(Ptr);
+	}
 }
 
 
