Removed Trace::Flush(). Upcoming changes will eliminate any latency between writing and consuming events.

diff --git a/Engine/Source/Runtime/Launch/Private/LaunchEngineLoop.cpp b/Engine/Source/Runtime/Launch/Private/LaunchEngineLoop.cpp
--- a/Engine/Source/Runtime/Launch/Private/LaunchEngineLoop.cpp
+++ b/Engine/Source/Runtime/Launch/Private/LaunchEngineLoop.cpp
@@ -5254,8 +5254,6 @@ void FEngineLoop::AppPreExit( )
 		GShaderCompilingManager = nullptr;
 	}
 #endif
-
-	Trace::Flush();
 }
 
 
diff --git a/Engine/Source/Runtime/TraceLog/Private/Trace/Trace.cpp b/Engine/Source/Runtime/TraceLog/Private/Trace/Trace.cpp
--- a/Engine/Source/Runtime/TraceLog/Private/Trace/Trace.cpp
+++ b/Engine/Source/Runtime/TraceLog/Private/Trace/Trace.cpp
@@ -12,7 +12,6 @@ namespace Private
 {
 
 ////////////////////////////////////////////////////////////////////////////////
-void	Writer_Flush();
 bool	Writer_Connect(const ANSICHAR*);
 bool	Writer_Open(const ANSICHAR*);
 uint32	Writer_EventToggle(const ANSICHAR*, bool);
@@ -61,12 +60,6 @@ uint32 ToggleEvent(const TCHAR* Wildcard, bool bState)
 	return Private::Writer_EventToggle(WildcardA, bState);
 }
 
-////////////////////////////////////////////////////////////////////////////////
-void Flush()
-{
-	Private::Writer_Flush();
-}
-
 } // namespace Trace
 
 #else
diff --git a/Engine/Source/Runtime/TraceLog/Private/Trace/Writer.cpp b/Engine/Source/Runtime/TraceLog/Private/Trace/Writer.cpp
--- a/Engine/Source/Runtime/TraceLog/Private/Trace/Writer.cpp
+++ b/Engine/Source/Runtime/TraceLog/Private/Trace/Writer.cpp
@@ -221,31 +221,6 @@ static void Writer_UpdateBuffers(void (*DataSink)(const uint8*, uint32))
 	}
 }
 
-////////////////////////////////////////////////////////////////////////////////
-void Writer_Flush()
-{
-	using namespace Private;
-
-	FBuffer* Buffer = GActiveBuffer.load(std::memory_order_acquire);
-	uint32 PrevUsed = Buffer->Used.fetch_add(BufferSize, std::memory_order_relaxed);
-
-	if (!(PrevUsed & BufferSize))
-	{
-		FBuffer* NextBuffer;
-		for (; ; Writer_Yield())
-		{
-			NextBuffer = Buffer->Next.load(std::memory_order_relaxed);
-			if (NextBuffer != nullptr)
-			{
-				break;
-			}
-		}
-
-		Buffer->Final = PrevUsed & BufferSizeMask;
-		GActiveBuffer.store(NextBuffer, std::memory_order_release);
-	}
-}
-
 ////////////////////////////////////////////////////////////////////////////////
 UE_TRACE_API void* Writer_NextBuffer(Private::FBuffer* Buffer, uint32 PrevUsed, uint32 SizeAndRef)
 {
@@ -836,8 +811,6 @@ static void Writer_Shutdown()
 		return;
 	}
 
-	Writer_Flush();
-
 	GWorkerThreadQuit = true;
 	ThreadJoin(GWorkerThread);
 	ThreadDestroy(GWorkerThread);
diff --git a/Engine/Source/Runtime/TraceLog/Public/Trace/Trace.h b/Engine/Source/Runtime/TraceLog/Public/Trace/Trace.h
--- a/Engine/Source/Runtime/TraceLog/Public/Trace/Trace.h
+++ b/Engine/Source/Runtime/TraceLog/Public/Trace/Trace.h
@@ -12,7 +12,6 @@ namespace Trace
 UE_TRACE_API bool	Connect(const TCHAR* Host) UE_TRACE_IMPL(false);
 UE_TRACE_API bool	Open(const TCHAR* Path) UE_TRACE_IMPL(false);
 UE_TRACE_API uint32 ToggleEvent(const TCHAR* Wildcard, bool bState) UE_TRACE_IMPL(false);
-UE_TRACE_API void	Flush() UE_TRACE_IMPL();
 
 } // namespace Trace
 
