Strategic movement of some function prototypes in to a header so we're in a better position in the future to improve the API a little.

diff --git a/Engine/Source/Runtime/TraceLog/Private/Trace/Platform.h b/Engine/Source/Runtime/TraceLog/Private/Trace/Platform.h
new file mode 100644
--- /dev/null
+++ b/Engine/Source/Runtime/TraceLog/Private/Trace/Platform.h
@@ -0,0 +1,46 @@
+// Copyright 1998-2019 Epic Games, Inc. All Rights Reserved.
+
+#pragma once
+
+#include "Trace/Private/Trace.h"
+
+#if UE_TRACE_ENABLED
+
+namespace Trace {
+namespace Private {
+
+////////////////////////////////////////////////////////////////////////////////
+uint8*	MemoryReserve(SIZE_T Size);
+void	MemoryFree(void* Address, SIZE_T Size);
+void	MemoryMap(void* Address, SIZE_T Size);
+void	MemoryUnmap(void* Address, SIZE_T Size);
+
+////////////////////////////////////////////////////////////////////////////////
+UPTRINT	ThreadCreate(const ANSICHAR* Name, void (*Entry)());
+uint32	ThreadGetCurrentId();
+void	ThreadSleep(uint32 Milliseconds);
+void	ThreadJoin(UPTRINT Handle);
+void	ThreadDestroy(UPTRINT Handle);
+
+////////////////////////////////////////////////////////////////////////////////
+uint64	TimeGetFrequency();
+uint64	TimeGetTimestamp();
+
+////////////////////////////////////////////////////////////////////////////////
+UPTRINT	TcpSocketConnect(const ANSICHAR* Host, uint16 Port);
+UPTRINT	TcpSocketListen(uint16 Port);
+int32	TcpSocketAccept(UPTRINT Socket, UPTRINT& Out);
+bool	TcpSocketHasData(UPTRINT Socket);
+
+////////////////////////////////////////////////////////////////////////////////
+int32	IoRead(UPTRINT Handle, void* Data, uint32 Size);
+bool	IoWrite(UPTRINT Handle, const void* Data, uint32 Size);
+void	IoClose(UPTRINT Handle);
+
+////////////////////////////////////////////////////////////////////////////////
+UPTRINT	FileOpen(const ANSICHAR* Path, const ANSICHAR* Mode);
+
+} // namespace Private
+} // namespace Trace
+
+#endif // UE_TRACE_ENABLED
diff --git a/Engine/Source/Runtime/TraceLog/Private/Trace/Writer.cpp b/Engine/Source/Runtime/TraceLog/Private/Trace/Writer.cpp
--- a/Engine/Source/Runtime/TraceLog/Private/Trace/Writer.cpp
+++ b/Engine/Source/Runtime/TraceLog/Private/Trace/Writer.cpp
@@ -4,6 +4,7 @@
 
 #if UE_TRACE_ENABLED
 
+#include "Trace/Platform.h"
 #include "Trace/Trace.h"
 #include "Misc/CString.h"
 
@@ -19,28 +20,6 @@
 namespace Trace {
 namespace Private {
 
-////////////////////////////////////////////////////////////////////////////////
-uint8*	MemoryReserve(SIZE_T);
-void	MemoryFree(void*, SIZE_T);
-void	MemoryMap(void*, SIZE_T);
-UPTRINT	TcpSocketConnect(const ANSICHAR*, uint16);
-UPTRINT	TcpSocketListen(uint16);
-int32	TcpSocketAccept(UPTRINT, UPTRINT&);
-bool	TcpSocketHasData(UPTRINT);
-int32	IoRead(UPTRINT, void*, uint32);
-bool	IoWrite(UPTRINT, const void*, uint32);
-void	IoClose(UPTRINT);
-UPTRINT	FileOpen(const ANSICHAR*, const ANSICHAR*);
-UPTRINT	ThreadCreate(const ANSICHAR*, void (*)());
-uint32	ThreadGetCurrentId();
-void	ThreadSleep(uint32 Milliseconds);
-void	ThreadJoin(UPTRINT);
-void	ThreadDestroy(UPTRINT);
-uint64	TimeGetFrequency();
-uint64	TimeGetTimestamp();
-
-
-
 ////////////////////////////////////////////////////////////////////////////////
 inline void Writer_Yield()
 {
