Removed TraceLog's instrumentation.

diff --git a/Engine/Source/Runtime/TraceLog/Private/Trace/Writer.cpp b/Engine/Source/Runtime/TraceLog/Private/Trace/Writer.cpp
--- a/Engine/Source/Runtime/TraceLog/Private/Trace/Writer.cpp
+++ b/Engine/Source/Runtime/TraceLog/Private/Trace/Writer.cpp
@@ -28,33 +28,6 @@ inline void Writer_Yield()
 
 
 
-////////////////////////////////////////////////////////////////////////////////
-UE_TRACE_EVENT_BEGIN($Trace, PerfWorker)
-	UE_TRACE_EVENT_FIELD(uint64, Start)
-	UE_TRACE_EVENT_FIELD(uint32, Acquire)
-	UE_TRACE_EVENT_FIELD(uint32, Send)
-	UE_TRACE_EVENT_FIELD(uint32, Done)
-UE_TRACE_EVENT_END()
-
-UE_TRACE_EVENT_BEGIN($Trace, PerfNextBuffer)
-	UE_TRACE_EVENT_FIELD(uint64, Start)
-	UE_TRACE_EVENT_FIELD(uint32, Acquire)
-	UE_TRACE_EVENT_FIELD(uint32, Done)
-	UE_TRACE_EVENT_FIELD(uint16, ThreadId)
-UE_TRACE_EVENT_END()
-
-void Writer_InitializeInstrumentation()
-{
-	// This is usually taken care of automatically by the UE_TRACE_* macros but
-	// as we're logging these events from within trace where there might not be
-	// any buffer space free, we'll explicitly initialise them here where we
-	// can guarantee we have available buffer space.
-	F$TracePerfNextBufferFields::Initialize();
-	F$TracePerfWorkerFields::Initialize();
-}
-
-
-
 ////////////////////////////////////////////////////////////////////////////////
 static uint64 GStartCycle = 0;
 
@@ -130,8 +103,6 @@ static void Writer_ShutdownBuffers()
 ////////////////////////////////////////////////////////////////////////////////
 static void Writer_RetireBuffer(void (*DataSink)(const uint8*, uint32))
 {
-	uint64 StartTsc = Writer_GetTimestamp();
-
 	uint32 TailUsed;
 	for (; ; Writer_Yield())
 	{
@@ -141,14 +112,12 @@ static void Writer_RetireBuffer(void (*DataSink)(const uint8*, uint32))
 			break;
 		}
 	}
-	uint64 AcquireTailTsc = Writer_GetTimestamp();
 
 	if (uint32 SendSize = GTailBuffer->Final - sizeof(FBuffer))
 	{
 		SendSize -= GTailPreSent;
 		DataSink(GTailBuffer->Data + GTailPreSent, SendSize);
 	}
-	uint64 SendTsc = Writer_GetTimestamp();
 
 	FBuffer* Next = GTailBuffer->Next.load(std::memory_order_relaxed);
 	new (GTailBuffer) FBuffer();
@@ -157,19 +126,6 @@ static void Writer_RetireBuffer(void (*DataSink)(const uint8*, uint32))
 	GHeadBuffer = GTailBuffer;
 	GTailBuffer = Next;
 	GTailPreSent = 0;
-	uint64 DoneTsc = Writer_GetTimestamp();
-
-	if (UE_TRACE_EVENT_IS_ENABLED($Trace, PerfWorker))
-	{
-		uint8 PerfEventBuffer[TRACE_PRIVATE_EVENT_SIZE($Trace, PerfWorker) + FEventDef::HeaderSize];
-		UE_TRACE_LOG($Trace, PerfWorker, PerfEventBuffer)
-			<< PerfWorker.Start(StartTsc)
-			<< PerfWorker.Acquire(uint32(AcquireTailTsc - StartTsc))
-			<< PerfWorker.Send(uint32(SendTsc - StartTsc))
-			<< PerfWorker.Done(uint32(DoneTsc - StartTsc));
-
-		DataSink(PerfEventBuffer, sizeof(PerfEventBuffer));
-	}
 }
 
 ////////////////////////////////////////////////////////////////////////////////
@@ -205,15 +161,6 @@ UE_TRACE_API void* Writer_NextBuffer(Private::FBuffer* Buffer, uint32 PrevUsed,
 {
 	using namespace Private;
 
-	uint32 PerfEventSize = 0;
-	if (UE_TRACE_EVENT_IS_ENABLED($Trace, PerfNextBuffer))
-	{
-		PerfEventSize = TRACE_PRIVATE_EVENT_SIZE($Trace, PerfNextBuffer) + FEventDef::HeaderSize;
-	}
-	SizeAndRef += PerfEventSize;
-
-	uint64 StartTsc = Writer_GetTimestamp();
-
 	FBuffer* NextBuffer;
 	while (true)
 	{
@@ -246,8 +193,6 @@ UE_TRACE_API void* Writer_NextBuffer(Private::FBuffer* Buffer, uint32 PrevUsed,
 		Buffer = NextBuffer;
 	}
 
-	uint64 AcquireTsc = Writer_GetTimestamp();
-
 	if (!(PrevUsed & BufferSize))
 	{
 		GActiveBuffer.compare_exchange_weak(Buffer, NextBuffer, std::memory_order_release);
@@ -256,17 +201,7 @@ UE_TRACE_API void* Writer_NextBuffer(Private::FBuffer* Buffer, uint32 PrevUsed,
 	PrevUsed &= BufferSizeMask;
 	uint8* Out = (uint8*)(UPTRINT(NextBuffer) + PrevUsed);
 
-	if (PerfEventSize)
-	{
-		uint64 DoneTsc = Writer_GetTimestamp();
-		UE_TRACE_LOG($Trace, PerfNextBuffer, Out)
-			<< PerfNextBuffer.Start(StartTsc)
-			<< PerfNextBuffer.Acquire(uint32(AcquireTsc - StartTsc))
-			<< PerfNextBuffer.Done(uint32(DoneTsc - StartTsc))
-			<< PerfNextBuffer.ThreadId(uint16(ThreadGetCurrentId()));
-	}
-
-	return Out + PerfEventSize;
+	return Out;
 }
 
 
@@ -779,7 +714,6 @@ static void Writer_InternalInitialize()
 
 	Writer_InitializeControl();
 	Writer_InitializeTiming();
-	Writer_InitializeInstrumentation();
 }
 
 ////////////////////////////////////////////////////////////////////////////////
