Fixed unexpected initialisation of GHoldBuffer instance causing small memory leak

diff --git a/Engine/Source/Runtime/TraceLog/Private/Trace/Writer.cpp b/Engine/Source/Runtime/TraceLog/Private/Trace/Writer.cpp
--- a/Engine/Source/Runtime/TraceLog/Private/Trace/Writer.cpp
+++ b/Engine/Source/Runtime/TraceLog/Private/Trace/Writer.cpp
@@ -207,7 +207,21 @@ UE_TRACE_API void* Writer_NextBuffer(Private::FBuffer* Buffer, uint32 PrevUsed,
 
 
 ////////////////////////////////////////////////////////////////////////////////
-class FHoldBuffer
+template <typename Class>
+class alignas(Class) TSafeStatic
+{
+public:
+	Class* operator -> ()				{ return (Class*)Buffer; }
+	Class const* operator -> () const	{ return (Class const*)Buffer; }
+
+protected:
+	uint8 alignas(Class) Buffer[sizeof(Class)];
+};
+
+
+
+////////////////////////////////////////////////////////////////////////////
+class FHoldBufferImpl
 {
 public:
 	void				Init();
@@ -227,45 +241,47 @@ private:
 	bool				bFull;
 };
 
+typedef TSafeStatic<FHoldBufferImpl> FHoldBuffer;
+
 ////////////////////////////////////////////////////////////////////////////////
-void FHoldBuffer::Init()
+void FHoldBufferImpl::Init()
 {
-	Base = MemoryReserve(FHoldBuffer::PageSize * FHoldBuffer::MaxPages);
+	Base = MemoryReserve(FHoldBufferImpl::PageSize * FHoldBufferImpl::MaxPages);
 	Used = 0;
 	MappedPageCount = 0;
 	bFull = false;
 }
 
 ////////////////////////////////////////////////////////////////////////////////
-void FHoldBuffer::Shutdown()
+void FHoldBufferImpl::Shutdown()
 {
 	if (Base == nullptr)
 	{
 		return;
 	}
 
-	MemoryFree(Base, FHoldBuffer::PageSize * FHoldBuffer::MaxPages);
+	MemoryFree(Base, FHoldBufferImpl::PageSize * FHoldBufferImpl::MaxPages);
 	Base = nullptr;
 	MappedPageCount = 0;
 	Used = 0;
 }
 
 ////////////////////////////////////////////////////////////////////////////////
-void FHoldBuffer::Write(const void* Data, uint32 Size)
+void FHoldBufferImpl::Write(const void* Data, uint32 Size)
 {
 	int32 NextUsed = Used + Size;
 
-	uint16 HotPageCount = uint16((NextUsed + (FHoldBuffer::PageSize - 1)) >> FHoldBuffer::PageShift);
+	uint16 HotPageCount = uint16((NextUsed + (FHoldBufferImpl::PageSize - 1)) >> FHoldBufferImpl::PageShift);
 	if (HotPageCount > MappedPageCount)
 	{
-		if (HotPageCount > FHoldBuffer::MaxPages)
+		if (HotPageCount > FHoldBufferImpl::MaxPages)
 		{
 			bFull = true;
 			return;
 		}
 
-		void* MapStart = Base + (UPTRINT(MappedPageCount) << FHoldBuffer::PageShift);
-		uint32 MapSize = (HotPageCount - MappedPageCount) << FHoldBuffer::PageShift;
+		void* MapStart = Base + (UPTRINT(MappedPageCount) << FHoldBufferImpl::PageShift);
+		uint32 MapSize = (HotPageCount - MappedPageCount) << FHoldBufferImpl::PageShift;
 		MemoryMap(MapStart, MapSize);
 
 		MappedPageCount = HotPageCount;
@@ -317,12 +333,12 @@ static void Writer_UpdateData()
 			bOk &= IoWrite(GDataHandle, &TransportHeader, sizeof(TransportHeader));
 
 			// Passively collected data
-			bOk &= IoWrite(GDataHandle, GHoldBuffer.GetData(), GHoldBuffer.GetSize());
+			bOk &= IoWrite(GDataHandle, GHoldBuffer->GetData(), GHoldBuffer->GetSize());
 
 			if (bOk)
 			{
 				GDataState = EDataState::Sending;
-				GHoldBuffer.Shutdown();
+				GHoldBuffer->Shutdown();
 			}
 			else
 			{
@@ -357,11 +373,11 @@ static void Writer_UpdateData()
 		// Send data to hold/ring
 		Writer_UpdateBuffers([] (const uint8* Data, uint32 Size)
 		{
-			GHoldBuffer.Write(Data, Size);
+			GHoldBuffer->Write(Data, Size);
 		});
 
 		// Did we overflow? Enter partial mode.
-		bool bOverflown = GHoldBuffer.IsFull();
+		bool bOverflown = GHoldBuffer->IsFull();
 		if (bOverflown && GDataState != EDataState::Partial)
 		{
 			GDataState = EDataState::Partial;
@@ -706,7 +722,7 @@ static void Writer_InternalInitialize()
 	GInitialized = true;
 
 	Writer_InitializeBuffers();
-	GHoldBuffer.Init();
+	GHoldBuffer->Init();
 
 	GWorkerThread = ThreadCreate("TraceWorker", Writer_WorkerThread);
 
@@ -730,7 +746,7 @@ static void Writer_Shutdown()
 
 	Writer_ShutdownControl();
 
-	GHoldBuffer.Shutdown();
+	GHoldBuffer->Shutdown();
 	Writer_ShutdownBuffers();
 
 	GInitialized = false;
