Fixed Clang compile error

diff --git a/Engine/Source/Runtime/TraceLog/Private/Trace/Writer.cpp b/Engine/Source/Runtime/TraceLog/Private/Trace/Writer.cpp
--- a/Engine/Source/Runtime/TraceLog/Private/Trace/Writer.cpp
+++ b/Engine/Source/Runtime/TraceLog/Private/Trace/Writer.cpp
@@ -215,7 +215,7 @@ public:
 	Class const* operator -> () const	{ return (Class const*)Buffer; }
 
 protected:
-	uint8 alignas(Class) Buffer[sizeof(Class)];
+	alignas(Class) uint8 Buffer[sizeof(Class)];
 };
 
 
