New thread-local buffering scheme for Trace events.

diff --git a/Engine/Source/Developer/TraceAnalysis/Private/Analysis/Engine.cpp b/Engine/Source/Developer/TraceAnalysis/Private/Analysis/Engine.cpp
--- a/Engine/Source/Developer/TraceAnalysis/Private/Analysis/Engine.cpp
+++ b/Engine/Source/Developer/TraceAnalysis/Private/Analysis/Engine.cpp
@@ -84,6 +84,142 @@ const uint8* FTransportReader::GetPointerImpl(uint32 BlockSize)
 
 
 
+////////////////////////////////////////////////////////////////////////////////
+class FTlsTransportReader
+	: public FTransportReader
+{
+public:
+	virtual					~FTlsTransportReader();
+	virtual void			Advance(uint32 BlockSize) override;
+	virtual const uint8*	GetPointerImpl(uint32 BlockSize) override;
+
+private:
+	struct FPayloadNode
+	{
+		FPayloadNode*		Next;
+		uint32				Cursor;
+		uint16				Serial;
+		uint16				Size;
+		uint8				Data[];
+	};
+
+	bool					GetNextBatch();
+	static const uint32		MaxPayloadSize = 8192;
+	FPayloadNode*			ActiveList = nullptr;
+	FPayloadNode*			PendingList = nullptr;
+	FPayloadNode*			FreeList = nullptr;
+};
+
+////////////////////////////////////////////////////////////////////////////////
+FTlsTransportReader::~FTlsTransportReader()
+{
+	for (FPayloadNode* Root : {ActiveList, PendingList, FreeList})
+	{
+		for (FPayloadNode* Node = Root; Node != nullptr;)
+		{
+			FPayloadNode* Next = Node->Next;
+			delete[] Node;
+			Node = Next;
+		}
+	}
+}
+
+////////////////////////////////////////////////////////////////////////////////
+void FTlsTransportReader::Advance(uint32 BlockSize)
+{
+	if (ActiveList != nullptr)
+	{
+		ActiveList->Cursor += BlockSize;
+	}
+}
+
+////////////////////////////////////////////////////////////////////////////////
+const uint8* FTlsTransportReader::GetPointerImpl(uint32 BlockSize)
+{
+	if (ActiveList == nullptr && !GetNextBatch())
+	{
+		return nullptr;
+	}
+
+	uint32 NextCursor = ActiveList->Cursor + BlockSize;
+	if (NextCursor > ActiveList->Size)
+	{
+		FPayloadNode* Node = ActiveList;
+		ActiveList = ActiveList->Next;
+		Node->Next = FreeList;
+		FreeList = Node;
+		return GetPointerImpl(BlockSize);
+	}
+
+	return ActiveList->Data + ActiveList->Cursor;
+}
+
+////////////////////////////////////////////////////////////////////////////////
+bool FTlsTransportReader::GetNextBatch()
+{
+	struct FPayload
+	{
+		int16	Serial;
+		uint16	Size;
+		uint8	Data[];
+	};
+
+	int16 LastSerial = -1;
+	if (PendingList != nullptr)
+	{
+		LastSerial = PendingList->Serial;
+	}
+
+	while (true)
+	{
+		const auto* Payload = (const FPayload*)FTransportReader::GetPointerImpl(sizeof(FPayload));
+		if (Payload == nullptr)
+		{
+			return false;
+		}
+
+		// If this new payload is part of the next event batch then we've finished
+		// building the current batch. The current batch can be activated.
+		if (LastSerial >= Payload->Serial)
+		{
+			ActiveList = PendingList;
+			PendingList = nullptr;
+			break;
+		}
+
+		if (FTransportReader::GetPointerImpl(Payload->Size) == nullptr)
+		{
+			return false;
+		}
+
+		FTransportReader::Advance(Payload->Size);
+
+		FPayloadNode* Node;
+		if (FreeList != nullptr)
+		{
+			Node = FreeList;
+			FreeList = Node->Next;
+		}
+		else
+		{
+			Node = (FPayloadNode*)GMalloc->Malloc(sizeof(FPayloadNode) + MaxPayloadSize);
+		}
+
+		Node->Serial = Payload->Serial;
+		Node->Cursor = 0;
+		Node->Size = uint16(Payload->Size - sizeof(*Payload));
+		Node->Next = PendingList;
+		FMemory::Memcpy(Node->Data, Payload->Data, Node->Size);
+
+		PendingList = Node;
+		LastSerial = Payload->Serial;
+	}
+
+	return true;
+}
+
+
+
 ////////////////////////////////////////////////////////////////////////////////
 #if 0
 class FLz4TransportReader
@@ -694,6 +830,7 @@ bool FAnalysisEngine::EstablishTransport(FStreamReader::FData& Data)
 	switch (Header->Format)
 	{
 	case 1:		Transport = new FTransportReader(); break;
+	case 2:		Transport = new FTlsTransportReader(); break;
 	default:	return false;
 	//case 'E':	/* See the magic above */ break;
 	//case 'T':	/* See the magic above */ break;
diff --git a/Engine/Source/Runtime/TraceLog/Private/Trace/Writer.cpp b/Engine/Source/Runtime/TraceLog/Private/Trace/Writer.cpp
--- a/Engine/Source/Runtime/TraceLog/Private/Trace/Writer.cpp
+++ b/Engine/Source/Runtime/TraceLog/Private/Trace/Writer.cpp
@@ -5,6 +5,7 @@
 #if UE_TRACE_ENABLED
 
 #include "Trace/Platform.h"
+#include "Trace/Private/Atomic.h"
 #include "Trace/Trace.h"
 #include "Misc/CString.h"
 
@@ -55,158 +56,150 @@ void Writer_InitializeTiming()
 
 
 ////////////////////////////////////////////////////////////////////////////////
-UE_TRACE_API TTraceAtomic<FBuffer*>	GActiveBuffer;
-static FBuffer*						GTailBuffer;
-static FBuffer*						GHeadBuffer;
-static void*						GAllocBase;
-static uint32						GAllocSize;
-static uint32						GTailPreSent;
+#define T_ALIGN alignas(PLATFORM_CACHE_LINE_SIZE)
+static uint8						GEmptyBuffer[sizeof(FWriteBuffer)];
+thread_local FWriteBuffer*			GWriteBuffer		= (FWriteBuffer*)GEmptyBuffer;
+T_ALIGN static void* volatile		GFirstEvent;
+UE_TRACE_API T_ALIGN void* volatile	GLastEvent;			// = nullptr;
+static const uint32					GPoolSize			= 384 << 20; // 384MB ought to be enough
+static const uint32					GPoolBlockSize		= 2 << 10;
+static const uint32					GPoolPageGrowth		= GPoolBlockSize << 5;
+static const uint32					GPoolInitPageSize	= GPoolBlockSize << 5;
+static uint8*						GPoolBase;			// = nullptr;
+T_ALIGN static uint8* volatile		GPoolPageCursor;	// = nullptr;
+T_ALIGN static void* volatile		GPoolFreeList;		// = nullptr;
+#undef T_ALIGN
 
 ////////////////////////////////////////////////////////////////////////////////
-static void Writer_InitializeBuffers()
-{
-	const uint32 BufferCount = 4;
-
-	GAllocSize = BufferSize * (BufferCount + 1);
-	GAllocBase = MemoryReserve(GAllocSize);
-
-	UPTRINT BufferBase = UPTRINT(GAllocBase);
-	BufferBase += BufferSizeMask;
-	BufferBase &= ~UPTRINT(BufferSizeMask);
-	MemoryMap((void*)BufferBase, BufferSize * BufferCount);
-
-	FBuffer* Buffers[BufferCount];
-	for (int i = 0; i < BufferCount; ++i)
-	{
-		void* Block = (void*)(BufferBase + (BufferSize * i));
-		Buffers[i] = new (Block) FBuffer();
-	}
-
-	for (int i = 1; i < BufferCount; ++i)
-	{
-		Buffers[i - 1]->Next = Buffers[i];
-	}
-
-	GTailBuffer = Buffers[0];
-	GHeadBuffer = Buffers[BufferCount - 1];
-	GTailPreSent = 0;
-
-	GActiveBuffer.store(Buffers[0], std::memory_order_release);
-}
-
-////////////////////////////////////////////////////////////////////////////////
-static void Writer_ShutdownBuffers()
+#if !IS_MONOLITHIC
+UE_TRACE_API FWriteBuffer* Writer_GetBuffer()
 {
-	MemoryFree(GAllocBase, GAllocSize);
+	// Thread locals and DLLs don't mix so for modular builds we are forced to
+	// export this function to access thread-local variables.
+	return GWriteBuffer;
 }
+#endif
 
 ////////////////////////////////////////////////////////////////////////////////
-static void Writer_RetireBuffer(void (*DataSink)(const uint8*, uint32))
+static FWriteBuffer* Writer_NextBufferInternal(uint32 PageGrowth)
 {
-	uint32 TailUsed;
-	for (; ; Writer_Yield())
+	// Fetch a new buffer
+	FWriteBuffer* Next;
+	while (true)
 	{
-		TailUsed = GTailBuffer->Used.load(std::memory_order_acquire);
-		if (TailUsed < BufferRefBit)
+		// First we'll try one from the free list
+		void* Owned = AtomicLoadRelaxed(&GPoolFreeList);
+		if (Owned != nullptr)
 		{
-			break;
+			if (!AtomicCompareExchangeRelaxed(&GPoolFreeList, *(void**)Owned, Owned))
+			{
+				Writer_Yield();
+				continue;
+			}
 		}
-	}
 
-	if (uint32 SendSize = GTailBuffer->Final - sizeof(FBuffer))
-	{
-		SendSize -= GTailPreSent;
-		DataSink(GTailBuffer->Data + GTailPreSent, SendSize);
-	}
-
-	FBuffer* Next = GTailBuffer->Next.load(std::memory_order_relaxed);
-	new (GTailBuffer) FBuffer();
-
-	GHeadBuffer->Next.store(GTailBuffer, std::memory_order_release);
-	GHeadBuffer = GTailBuffer;
-	GTailBuffer = Next;
-	GTailPreSent = 0;
-}
+		// If we didn't fetch the sentinal then we've taken a block we can use
+		if (Owned != nullptr)
+		{
+			Next = (FWriteBuffer*)Owned;
+			break;
+		}
 
-////////////////////////////////////////////////////////////////////////////////
-static void Writer_UpdateBuffers(void (*DataSink)(const uint8*, uint32))
-{
-	while (true)
-	{
-		FBuffer* Buffer = GActiveBuffer.load(std::memory_order_acquire);
-		if (GTailBuffer != Buffer)
+		// The free list is empty. Map some more memory.
+		uint8* PageBase = (uint8*)AtomicLoadRelaxed(&GPoolPageCursor);
+		if (!AtomicCompareExchangeAcquire(&GPoolPageCursor, PageBase + PageGrowth, PageBase))
 		{
-			Writer_RetireBuffer(DataSink);
+			// Someone else is mapping memory so we'll briefly yield and try the
+			// free list again.
+			Writer_Yield();
 			continue;
 		}
 
-		uint32 Used = Buffer->Used.load(std::memory_order_relaxed);
-		if (Used >= BufferSize) /* is the buffer in use somewhere? */
+		// We claimed the pool cursor so it is now our job to map memory and add
+		// it to the free list.
+		MemoryMap(PageBase, PageGrowth);
+
+		// The first block in the page we'll use for the next buffer
+		Next = (FWriteBuffer*)PageBase;
+		uint8* FirstBlock = PageBase + GPoolBlockSize;
+
+		// Link subsequent blocks together
+		uint8* Block = FirstBlock;
+		for (int i = 2, n = PageGrowth / GPoolBlockSize; i < n; ++i)
 		{
-			continue;
+			auto* Buffer = (FWriteBuffer*)Block;
+			Buffer->Next = (FWriteBuffer*)(Block + GPoolBlockSize);
+			Block += GPoolBlockSize;
 		}
 
-		if (uint32 Sendable = Used - sizeof(FBuffer) - GTailPreSent)
+		// And insert the block list into the freelist
+		uint8* LastBlock = Block;
+		for (void** ListNode = (void**)LastBlock;; Writer_Yield())
 		{
-			DataSink(Buffer->Data + GTailPreSent, Sendable);
-			GTailPreSent += Sendable;
+			*ListNode = AtomicLoadRelaxed(&GPoolFreeList);
+			if (AtomicCompareExchangeRelease(&GPoolFreeList, (void*)FirstBlock, *ListNode))
+			{
+				break;
+			}
 		}
 
 		break;
 	}
+
+	GWriteBuffer = Next;
+
+	Next->Cursor = ((uint8*)Next + GPoolBlockSize);
+	return Next;
 }
 
 ////////////////////////////////////////////////////////////////////////////////
-UE_TRACE_API void* Writer_NextBuffer(Private::FBuffer* Buffer, uint32 PrevUsed, uint32 SizeAndRef)
+UE_TRACE_API uint8* Writer_NextBuffer(uint16 Size)
 {
-	using namespace Private;
+	FWriteBuffer* Current = GWriteBuffer;
 
-	FBuffer* NextBuffer;
-	while (true)
+	// Retire current buffer unless its the initial boot one.
+	if (UPTRINT(Current) != UPTRINT(GEmptyBuffer))
 	{
-		if (!(PrevUsed & BufferSize))
-		{
-			Buffer->Final = PrevUsed & BufferSizeMask;
-		}
-
-		// Get the next candidate buffer to allocate from.
-		for (; ; Writer_Yield())
+		// To retire a buffer we'll link it into the event list which event
+		// consumption can detect and use to return the buffer to the free list.
+		for (;; Writer_Yield())
 		{
-			NextBuffer = Buffer->Next.load(std::memory_order_relaxed);
-			if (NextBuffer != nullptr)
+			void* Expected = AtomicLoadRelaxed(&GLastEvent);
+			*(void**)Current = Expected;
+			if (AtomicCompareExchangeRelease(&GLastEvent, (void*)Current, Expected))
 			{
 				break;
 			}
 		}
+	}
 
-		Buffer->Used.fetch_sub(BufferRefBit, std::memory_order_release);
+	FWriteBuffer* NextBuffer = Writer_NextBufferInternal(GPoolPageGrowth);
 
-		// Try and allocate some space in the next buffer.
-		PrevUsed = NextBuffer->Used.fetch_add(SizeAndRef, std::memory_order_relaxed);
-		uint32 Used = PrevUsed + SizeAndRef;
-		if (UNLIKELY(!(Used & BufferSize)))
-		{
-			break;
-		}
+	NextBuffer->Cursor -= Size;
+	return NextBuffer->Cursor;
+}
 
-		// Next buffer's full. Try again.
-		Buffer = NextBuffer;
-	}
+////////////////////////////////////////////////////////////////////////////////
+static void Writer_InitializeBuffers()
+{
+	GPoolBase = MemoryReserve(GPoolSize);
+	AtomicStoreRelaxed(&GPoolPageCursor, GPoolBase);
 
-	if (!(PrevUsed & BufferSize))
-	{
-		GActiveBuffer.compare_exchange_weak(Buffer, NextBuffer, std::memory_order_release);
-	}
+	Writer_NextBufferInternal(GPoolInitPageSize);
 
-	PrevUsed &= BufferSizeMask;
-	uint8* Out = (uint8*)(UPTRINT(NextBuffer) + PrevUsed);
+	static_assert(GPoolPageGrowth >= 0x10000, "Page growth must be >= 64KB");
+	static_assert(GPoolInitPageSize >= 0x10000, "Initial page size must be >= 64KB");
+}
 
-	return Out;
+////////////////////////////////////////////////////////////////////////////////
+static void Writer_ShutdownBuffers()
+{
+	MemoryFree(GPoolBase, GPoolSize);
 }
 
 
 
-////////////////////////////////////////////////////////////////////////////////
+////////////////////////////////////////////////////////////////////////////
 template <typename Class>
 class alignas(Class) TSafeStatic
 {
@@ -306,6 +299,136 @@ static UPTRINT			GDataHandle			= 0;
 static EDataState		GDataState;			// = EDataState::Passive;
 UPTRINT					GPendingDataHandle	= 0;
 
+////////////////////////////////////////////////////////////////////////////////
+static void Writer_ConsumeEvents()
+{
+	// Claim ownership of the latest chain of sent events.
+	void* LatestEvent;
+	for (;; Writer_Yield())
+	{
+		LatestEvent = AtomicLoadRelaxed(&GLastEvent);
+		if (AtomicCompareExchangeAcquire(&GLastEvent, (void*)nullptr, LatestEvent))
+		{
+			break;
+		}
+	}
+
+	FWriteBuffer* RetiredHead = nullptr;
+	FWriteBuffer* RetiredTail = nullptr;
+
+	struct FCollector
+	{
+		struct FPayload
+		{
+			struct FHeader
+			{
+				uint16	Serial;
+				uint16	Size; // including header
+			};
+			FHeader		Header;
+			uint8		Data[8192];
+		};
+
+		void Send(const void* Data, uint32 Size)
+		{
+			if (GDataState == EDataState::Sending)
+			{
+				// Transmit data to the io handle
+				if (GDataHandle)
+				{
+					if (!IoWrite(GDataHandle, Data, Size))
+					{
+						IoClose(GDataHandle);
+						GDataHandle = 0;
+					}
+				}
+			}
+			else
+			{
+				GHoldBuffer->Write(Data, Size);
+
+				// Did we overflow? Enter partial mode.
+				bool bOverflown = GHoldBuffer->IsFull();
+				if (bOverflown && GDataState != EDataState::Partial)
+				{
+					GDataState = EDataState::Partial;
+				}
+			}
+		}
+
+		void Flush()
+		{
+			// There will always be space remaining for the header because we've
+			// arranged for that by including the header in FPayload. We'll shift
+			// it forward so it butts up against the event data (at the expense of
+			// some occasional unaligned stores).
+			Cursor -= sizeof(FPayload::FHeader);
+
+			auto* Out = (FPayload::FHeader*)(Payload.Data + Cursor);
+			Out->Serial = Serial;
+			Out->Size = sizeof(FPayload::Data) - Cursor;
+			Send(Out, Out->Size);
+
+			Cursor = sizeof(FPayload::Data);
+			Serial++;
+		}
+
+		void Write(const void* Data, uint32 Size)
+		{
+			if (int32(Cursor - Size) < 0)
+			{
+				Flush();
+			}
+
+			Cursor -= Size;
+			memcpy(Payload.Data + Cursor, Data, Size);
+		}
+
+		int16		Cursor = sizeof(FPayload::Data);
+		uint16		Serial = 0;
+		FPayload	Payload;
+	};
+
+	FCollector Collector;
+	for (void* EventPtr = LatestEvent; EventPtr != nullptr; )
+	{
+		// Is this "event" a retired buffer?
+		if ((UPTRINT(EventPtr) & (GPoolBlockSize - 1)) == 0)
+		{
+			auto* Retiree = (FWriteBuffer*)EventPtr;
+			EventPtr = *(void**)EventPtr;
+
+			Retiree->Next = RetiredHead;
+			RetiredHead = Retiree;
+			RetiredTail = (RetiredTail != nullptr) ? RetiredTail : RetiredHead;
+			continue;
+		}
+
+		const uint16* Header = (uint16*)(UPTRINT(EventPtr) + sizeof(void*));
+		uint16 DataSize = Header[1] + sizeof(uint32);
+
+		Collector.Write(Header, DataSize);
+
+		EventPtr = *(void**)EventPtr;
+	}
+	Collector.Flush();
+
+	if (RetiredHead == nullptr)
+	{
+		return;
+	}
+
+	// Put the retirees we found back into the system again.
+	for (void** ListNode = (void**)RetiredTail;; Writer_Yield())
+	{
+		*ListNode = AtomicLoadRelaxed(&GPoolFreeList);
+		if (AtomicCompareExchangeRelease(&GPoolFreeList, (void*)RetiredHead, *ListNode))
+		{
+			break;
+		}
+	}
+}
+
 ////////////////////////////////////////////////////////////////////////////////
 static void Writer_UpdateData()
 {
@@ -329,7 +452,7 @@ static void Writer_UpdateData()
 			const struct {
 				uint8 Format;
 				uint8 Parameter;
-			} TransportHeader = { 1 };
+			} TransportHeader = { 2 };
 			bOk &= IoWrite(GDataHandle, &TransportHeader, sizeof(TransportHeader));
 
 			// Passively collected data
@@ -348,41 +471,7 @@ static void Writer_UpdateData()
 		}
 	}
 
-	// Passive mode?
-	if (GDataState == EDataState::Sending)
-	{
-		// Transmit data to the io handle
-		if (GDataHandle)
-		{
-			Writer_UpdateBuffers([] (const uint8* Data, uint32 Size)
-			{
-				if (GDataHandle && !IoWrite(GDataHandle, Data, Size))
-				{
-					IoClose(GDataHandle);
-					GDataHandle = 0;
-				}
-			});
-		}
-		else
-		{
-			Writer_UpdateBuffers([] (const uint8*, uint32) {});
-		}
-	}
-	else
-	{
-		// Send data to hold/ring
-		Writer_UpdateBuffers([] (const uint8* Data, uint32 Size)
-		{
-			GHoldBuffer->Write(Data, Size);
-		});
-
-		// Did we overflow? Enter partial mode.
-		bool bOverflown = GHoldBuffer->IsFull();
-		if (bOverflown && GDataState != EDataState::Partial)
-		{
-			GDataState = EDataState::Partial;
-		}
-	}
+	Writer_ConsumeEvents();
 }
 
 
@@ -684,12 +773,14 @@ static void Writer_WorkerThread()
 {
 	while (!GWorkerThreadQuit)
 	{
-		const uint32 SleepMs = 60;
+		const uint32 SleepMs = 24;
 		ThreadSleep(SleepMs);
 
 		Writer_UpdateControl();
 		Writer_UpdateData();
 	}
+
+	Writer_ConsumeEvents();
 }
 
 
@@ -722,12 +813,12 @@ static void Writer_InternalInitialize()
 	GInitialized = true;
 
 	Writer_InitializeBuffers();
+	Writer_LogHeader();
+
 	GHoldBuffer->Init();
 
 	GWorkerThread = ThreadCreate("TraceWorker", Writer_WorkerThread);
 
-	Writer_LogHeader();
-
 	Writer_InitializeControl();
 	Writer_InitializeTiming();
 }
diff --git a/Engine/Source/Runtime/TraceLog/Public/Trace/Private/Atomic.h b/Engine/Source/Runtime/TraceLog/Public/Trace/Private/Atomic.h
new file mode 100644
--- /dev/null
+++ b/Engine/Source/Runtime/TraceLog/Public/Trace/Private/Atomic.h
@@ -0,0 +1,182 @@
+// Copyright 1998-2019 Epic Games, Inc. All Rights Reserved.
+
+#pragma once
+
+namespace Trace {
+namespace Private {
+
+////////////////////////////////////////////////////////////////////////////////
+template <typename Type> Type	AtomicLoadRelaxed(Type volatile* Source);
+template <typename Type> Type	AtomicLoadAcquire(Type volatile* Source);
+template <typename Type> void	AtomicStoreRelaxed(Type volatile* Target, Type Value);
+template <typename Type> void	AtomicStoreRelease(Type volatile* Target, Type Value);
+template <typename Type> bool	AtomicCompareExchangeRelaxed(Type volatile* Target, Type New, Type Expected);
+template <typename Type> bool	AtomicCompareExchangeAcquire(Type volatile* Target, Type New, Type Expected);
+template <typename Type> bool	AtomicCompareExchangeRelease(Type volatile* Target, Type New, Type Expected);
+
+} // namespace Private
+} // namespace Trace
+
+////////////////////////////////////////////////////////////////////////////////
+#define IS_MSVC					0
+#define IS_GCC_COMPATIBLE		0
+#if defined(_MSC_VER) && !defined(__clang__)
+#	undef  IS_MSVC
+#	define IS_MSVC				1
+#elif defined(__clang__) || defined(__GNUC__)
+#	undef  IS_GCC_COMPATIBLE
+#	define IS_GCC_COMPATIBLE	1
+#endif
+
+////////////////////////////////////////////////////////////////////////////////
+#if IS_MSVC
+#	include <intrin.h>
+#	if defined(_M_ARM) || defined(_M_ARM64)
+#		if defined(_M_ARM)
+#			include <armintr.h>
+#		else
+#			include <arm64intr.h>
+#		endif
+#	endif
+#endif
+
+namespace Trace {
+namespace Private {
+
+////////////////////////////////////////////////////////////////////////////////
+#if IS_MSVC
+#	if defined(_M_ARM) || defined(_M_ARM64)
+#		define INTERLOCKED_API(Name, Suffix, ...)	Name##_##Suffix(__VA_ARGS__)
+#	elif defined(_M_IX86) || defined(_M_X64)
+#		define INTERLOCKED_API(Name, Suffix, ...)	Name(__VA_ARGS__)
+#	endif
+#endif
+
+////////////////////////////////////////////////////////////////////////////////
+template <typename Type>
+inline Type AtomicLoadRelaxed(Type volatile* Source)
+{
+	static_assert(sizeof(Type) == sizeof(void*), "");
+#if IS_MSVC
+#	if defined(_M_ARM)
+		__int32 Value = __iso_volatile_load32((__int32 volatile*)Source);
+		return (Type*)Value;
+#	elif defined(_M_ARM64)
+		__int64 Value = __iso_volatile_load64((__int64 volatile*)Source);
+		return (Type*)Value;
+#	else
+		return *Source;
+#	endif
+#elif IS_GCC_COMPATIBLE
+	return __atomic_load_n(Source, __ATOMIC_RELAXED);
+#endif
+}
+
+////////////////////////////////////////////////////////////////////////////////
+template <typename Type>
+inline Type AtomicLoadAcquire(Type volatile* Source)
+{
+	static_assert(sizeof(Type) == sizeof(void*), "");
+#if IS_MSVC
+#	if defined(_M_ARM)
+		__int32 Value = __iso_volatile_load32((__int32 volatile*)Source);
+		__dmb(_ARM_BARRIER_ISH);
+		return Type(Value);
+#	elif defined(_M_ARM64)
+		__int64 Value = __iso_volatile_load64((__int64 volatile*)Source);
+		__dmb(_ARM64_BARRIER_ISH);
+		return Type(Value);
+#	else
+		Type Value = *Source;
+		_ReadWriteBarrier();
+		return Value;
+#	endif
+#elif IS_GCC_COMPATIBLE
+	return __atomic_load_n(Source, __ATOMIC_ACQUIRE);
+#endif
+}
+
+////////////////////////////////////////////////////////////////////////////////
+template <typename Type>
+inline void AtomicStoreRelaxed(Type volatile* Target, Type Value)
+{
+	static_assert(sizeof(Type) == sizeof(void*), "");
+#if IS_MSVC
+#	if defined(_M_ARM)
+		__iso_volatile_store32((__int32 volatile*)Target, UPTRINT(Value));
+#	elif defined(_M_ARM64)
+		__iso_volatile_store64((__int64 volatile*)Target, UPTRINT(Value));
+#	else
+		*Target = Value;
+#	endif
+#elif IS_GCC_COMPATIBLE
+	return __atomic_store_n(Target, Value, __ATOMIC_RELAXED);
+#endif
+}
+
+////////////////////////////////////////////////////////////////////////////////
+template <typename Type>
+inline void AtomicStoreRelease(Type volatile* Target, Type Value)
+{
+	static_assert(sizeof(Type) == sizeof(void*), "");
+#if IS_MSVC
+#	if defined(_M_ARM)
+		__dmb(_ARM_BARRIER_ISH);
+		__iso_volatile_store32((__int32 volatile*)Target, __int32(Value));
+#	elif defined(_M_ARM64)
+		__dmb(_ARM64_BARRIER_ISH);
+		__iso_volatile_store64((__int64 volatile*)Target, __int64(Value));
+#	else
+		_ReadWriteBarrier();
+		*Target = Value;
+#	endif
+#elif IS_GCC_COMPATIBLE
+	return __atomic_store_n(Target, Value, __ATOMIC_RELEASE);
+#endif
+}
+
+////////////////////////////////////////////////////////////////////////////////
+template <typename Type>
+inline bool AtomicCompareExchangeRelaxed(Type volatile* Target, Type New, Type Expected)
+{
+	static_assert(sizeof(Type) == sizeof(void*), "");
+#if IS_MSVC
+	return INTERLOCKED_API(_InterlockedCompareExchangePointer, _nf, (void* volatile*)Target, New, Expected) == Expected;
+#elif IS_GCC_COMPATIBLE
+	Type InOut = Expected;
+	return __atomic_compare_exchange_n(Target, &InOut, New, true, __ATOMIC_RELAXED, __ATOMIC_RELAXED);
+#endif
+}
+
+////////////////////////////////////////////////////////////////////////////////
+template <typename Type>
+inline bool AtomicCompareExchangeAcquire(Type volatile* Target, Type New, Type Expected)
+{
+	static_assert(sizeof(Type) == sizeof(void*), "");
+#if IS_MSVC
+	return INTERLOCKED_API(_InterlockedCompareExchangePointer, _acq, (void* volatile*)Target, New, Expected) == Expected;
+#elif IS_GCC_COMPATIBLE
+	Type InOut = Expected;
+	return __atomic_compare_exchange_n(Target, &InOut, New, true, __ATOMIC_ACQUIRE, __ATOMIC_RELAXED);
+#endif
+}
+
+////////////////////////////////////////////////////////////////////////////////
+template <typename Type>
+inline bool AtomicCompareExchangeRelease(Type volatile* Target, Type New, Type Expected)
+{
+	static_assert(sizeof(Type) == sizeof(void*), "");
+#if IS_MSVC
+	return INTERLOCKED_API(_InterlockedCompareExchangePointer, _rel, (void* volatile*)Target, New, Expected) == Expected;
+#elif IS_GCC_COMPATIBLE
+	Type InOut = Expected;
+	return __atomic_compare_exchange_n(Target, &InOut, New, true, __ATOMIC_RELEASE, __ATOMIC_RELAXED);
+#endif
+}
+
+#undef INTERLOCKED_API
+#undef IS_GCC_COMPATIBLE
+#undef IS_MSVC
+
+} // namespace Private
+} // namespace Trace
diff --git a/Engine/Source/Runtime/TraceLog/Public/Trace/Private/Writer.inl b/Engine/Source/Runtime/TraceLog/Public/Trace/Private/Writer.inl
--- a/Engine/Source/Runtime/TraceLog/Public/Trace/Private/Writer.inl
+++ b/Engine/Source/Runtime/TraceLog/Public/Trace/Private/Writer.inl
@@ -5,6 +5,7 @@
 #if UE_TRACE_ENABLED
 
 #include <atomic>
+#include "Atomic.h"
 
 namespace Trace
 {
@@ -22,40 +23,32 @@ namespace Private
 #endif
 
 ////////////////////////////////////////////////////////////////////////////////
-struct FBuffer
+struct FWriteBuffer
 {
-	FBuffer()
-	: Next(nullptr)
-	, Used(sizeof(FBuffer))
-	{
-	}
-
 	union
 	{
-		TTraceAtomic<FBuffer*>	Next;
-		char					Padding0[PLATFORM_CACHE_LINE_SIZE];
+		uint8*			Cursor;
+		FWriteBuffer*	Next;
 	};
-	union
-	{
-		TTraceAtomic<uint32>	Used;
-		char					Padding1[PLATFORM_CACHE_LINE_SIZE];
-	};
-	uint32						Final;
-	alignas(void*) uint8		Data[];
+	uint8				Data[];
 };
 
 #if defined(_MSC_VER)
 	#pragma warning(pop)
 #endif
 
-extern UE_TRACE_API TTraceAtomic<FBuffer*> GActiveBuffer;
-
-static const uint16		BufferSizePow2	= 19;
-static const uint32		BufferSize		= 1 << BufferSizePow2;
-static const uint32		BufferSizeMask	= BufferSize - 1;
-static const uint32		BufferRefBit	= BufferSize << 1;
+////////////////////////////////////////////////////////////////////////////////
+extern UE_TRACE_API void* volatile	GLastEvent;
+UE_TRACE_API uint8*					Writer_NextBuffer(uint16);
+UE_TRACE_API FWriteBuffer*			Writer_GetBuffer();
 
-extern UE_TRACE_API void* Writer_NextBuffer(FBuffer*, uint32, uint32);
+#if IS_MONOLITHIC
+extern thread_local FWriteBuffer* GWriteBuffer;
+inline FWriteBuffer* Writer_GetBuffer()
+{
+	return GWriteBuffer;
+}
+#endif
 
 } // Private
 
@@ -64,24 +57,17 @@ inline uint8* Writer_BeginLog(uint16 EventUid, uint16 Size)
 {
 	using namespace Private;
 
-	uint32 AllocSize = Size;
-	AllocSize += sizeof(uint32);
-	AllocSize += BufferRefBit;
-
-	// Fetch buffer and claim some space in it.
-	FBuffer* Buffer = GActiveBuffer.load(std::memory_order_acquire);
-	uint32 PrevUsed = Buffer->Used.fetch_add(AllocSize, std::memory_order_relaxed);
-	uint32 Used = PrevUsed + AllocSize;
-
-	uint32* Out = (uint32*)(UPTRINT(Buffer) + (BufferSizeMask & PrevUsed));
+	static const uint32 HeaderSize = sizeof(void*) + sizeof(uint32);
+	uint32 AllocSize = ((Size + HeaderSize) + 7) & ~7;
 
-	// Did we exhaust the active buffer?
-	if (UNLIKELY(Used & BufferSize))
+	FWriteBuffer* Buffer = Writer_GetBuffer();
+	uint8* Cursor = (Buffer->Cursor -= AllocSize);
+	if (UNLIKELY(PTRINT(Cursor) < PTRINT(Buffer->Data)))
 	{
-		Out = (uint32*)Writer_NextBuffer(Buffer, PrevUsed, AllocSize);
-		Buffer = (FBuffer*)(UPTRINT(Out) & ~UPTRINT(BufferSizeMask));
+		Cursor = Writer_NextBuffer(AllocSize);
 	}
 
+	uint32* Out = (uint32*)(Cursor + sizeof(void*));
 	Out[0] = (uint32(Size) << 16)|uint32(EventUid);
 	return (uint8*)(Out + 1);
 }
@@ -90,8 +76,19 @@ inline uint8* Writer_BeginLog(uint16 EventUid, uint16 Size)
 inline void Writer_EndLog(uint8* EventData)
 {
 	using namespace Private;
-	auto* Buffer = (FBuffer*)(UPTRINT(EventData) & ~UPTRINT(BufferSizeMask));
-	Buffer->Used.fetch_sub(BufferRefBit, std::memory_order_release);
+
+	EventData -= sizeof(void*) + sizeof(uint32);
+
+	// Add the event into the master linked list of events.
+	while (true)
+	{
+		void* Expected = AtomicLoadRelaxed(&GLastEvent);
+		*(void**)EventData = Expected;
+		if (AtomicCompareExchangeRelease(&GLastEvent, (void*)EventData, Expected))
+		{
+			break;
+		}
+	}
 }
 
 } // namespace Trace
