Assign a thread identifier to each buffer of trace events. A step towards being able to implicitly track which thread an event came from during analysis.

diff --git a/Engine/Source/Runtime/TraceLog/Private/Trace/Writer.cpp b/Engine/Source/Runtime/TraceLog/Private/Trace/Writer.cpp
--- a/Engine/Source/Runtime/TraceLog/Private/Trace/Writer.cpp
+++ b/Engine/Source/Runtime/TraceLog/Private/Trace/Writer.cpp
@@ -62,6 +62,7 @@ thread_local FWriteBuffer*			GWriteBuffer		= (FWriteBuffer*)GEmptyBuffer;
 T_ALIGN static void* volatile		GFirstEvent;
 UE_TRACE_API T_ALIGN void* volatile	GLastEvent;			// = nullptr;
 static const uint32					GPoolSize			= 384 << 20; // 384MB ought to be enough
+T_ALIGN static UPTRINT volatile		GThreadId;			// = 0;
 static const uint32					GPoolBlockSize		= 2 << 10;
 static const uint32					GPoolPageGrowth		= GPoolBlockSize << 5;
 static const uint32					GPoolInitPageSize	= GPoolBlockSize << 5;
@@ -157,6 +158,26 @@ UE_TRACE_API uint8* Writer_NextBuffer(uint16 Size)
 {
 	FWriteBuffer* Current = GWriteBuffer;
 
+	// Carry along or assign a new thread id
+	uint32 ThreadId;
+	if (UPTRINT(Current) == UPTRINT(GEmptyBuffer))
+	{
+		for (;; Writer_Yield())
+		{
+			UPTRINT CurrentTid = UPTRINT(AtomicLoadRelaxed((void* volatile*)&GThreadId));
+			UPTRINT NextTid = CurrentTid + 1;
+			if (AtomicCompareExchangeRelaxed((void* volatile*)&GThreadId, (void*)NextTid, (void*)CurrentTid))
+			{
+				ThreadId = uint32(CurrentTid);
+				break;
+			}
+		}
+	}
+	else
+	{
+		ThreadId = Current->ThreadId;
+	}
+
 	// Retire current buffer unless its the initial boot one.
 	if (UPTRINT(Current) != UPTRINT(GEmptyBuffer))
 	{
diff --git a/Engine/Source/Runtime/TraceLog/Public/Trace/Private/Writer.inl b/Engine/Source/Runtime/TraceLog/Public/Trace/Private/Writer.inl
--- a/Engine/Source/Runtime/TraceLog/Public/Trace/Private/Writer.inl
+++ b/Engine/Source/Runtime/TraceLog/Public/Trace/Private/Writer.inl
@@ -30,6 +30,7 @@ struct FWriteBuffer
 		uint8*			Cursor;
 		FWriteBuffer*	Next;
 	};
+	uint32				ThreadId;
 	uint8				Data[];
 };
 
