Fixed overly long UE_LOG output stomping trace's internal buffers.

diff --git a/Engine/Source/Runtime/Core/Public/Logging/LogTrace.h b/Engine/Source/Runtime/Core/Public/Logging/LogTrace.h
--- a/Engine/Source/Runtime/Core/Public/Logging/LogTrace.h
+++ b/Engine/Source/Runtime/Core/Public/Logging/LogTrace.h
@@ -23,7 +23,7 @@ struct FLogTrace
 	template <typename... Types>
 	static void OutputLogMessage(const void* LogPoint, Types... FormatArgs)
 	{
-		uint8 FormatArgsBuffer[4096];
+		uint8 FormatArgsBuffer[3072];
 		uint16 FormatArgsSize = FFormatArgsTrace::EncodeArguments(FormatArgsBuffer, FormatArgs...);
 		if (FormatArgsSize)
 		{
diff --git a/Engine/Source/Runtime/TraceLog/Private/Trace/Writer.cpp b/Engine/Source/Runtime/TraceLog/Private/Trace/Writer.cpp
--- a/Engine/Source/Runtime/TraceLog/Private/Trace/Writer.cpp
+++ b/Engine/Source/Runtime/TraceLog/Private/Trace/Writer.cpp
@@ -63,7 +63,7 @@ T_ALIGN static void* volatile		GFirstEvent;
 UE_TRACE_API T_ALIGN void* volatile	GLastEvent;			// = nullptr;
 static const uint32					GPoolSize			= 384 << 20; // 384MB ought to be enough
 T_ALIGN static UPTRINT volatile		GThreadId;			// = 0;
-static const uint32					GPoolBlockSize		= 2 << 10;
+static const uint32					GPoolBlockSize		= 4 << 10;
 static const uint32					GPoolPageGrowth		= GPoolBlockSize << 5;
 static const uint32					GPoolInitPageSize	= GPoolBlockSize << 5;
 static uint8*						GPoolBase;			// = nullptr;
