Generally speaking it is useful to assign the thread ID to something after you have created it.

diff --git a/Engine/Source/Runtime/TraceLog/Private/Trace/Writer.cpp b/Engine/Source/Runtime/TraceLog/Private/Trace/Writer.cpp
--- a/Engine/Source/Runtime/TraceLog/Private/Trace/Writer.cpp
+++ b/Engine/Source/Runtime/TraceLog/Private/Trace/Writer.cpp
@@ -195,6 +195,7 @@ UE_TRACE_API uint8* Writer_NextBuffer(uint16 Size)
 	}
 
 	FWriteBuffer* NextBuffer = Writer_NextBufferInternal(GPoolPageGrowth);
+	NextBuffer->ThreadId = ThreadId;
 
 	NextBuffer->Cursor -= Size;
 	return NextBuffer->Cursor;
