As TraceLog will now does it's own thread id, there is no need to be able to query the OS for a thread's identifier anymore.

diff --git a/Engine/Source/Runtime/TraceLog/Private/Trace/Detail/Windows/WindowsTrace.cpp b/Engine/Source/Runtime/TraceLog/Private/Trace/Detail/Windows/WindowsTrace.cpp
--- a/Engine/Source/Runtime/TraceLog/Private/Trace/Detail/Windows/WindowsTrace.cpp
+++ b/Engine/Source/Runtime/TraceLog/Private/Trace/Detail/Windows/WindowsTrace.cpp
@@ -63,12 +63,6 @@ UPTRINT ThreadCreate(const ANSICHAR* Name, void (*Entry)())
 	return UPTRINT(Handle);
 }
 
-////////////////////////////////////////////////////////////////////////////////
-uint32 ThreadGetCurrentId()
-{
-	return GetCurrentThreadId();
-}
-
 ////////////////////////////////////////////////////////////////////////////////
 void ThreadSleep(uint32 Milliseconds)
 {
diff --git a/Engine/Source/Runtime/TraceLog/Private/Trace/Platform.h b/Engine/Source/Runtime/TraceLog/Private/Trace/Platform.h
--- a/Engine/Source/Runtime/TraceLog/Private/Trace/Platform.h
+++ b/Engine/Source/Runtime/TraceLog/Private/Trace/Platform.h
@@ -17,7 +17,6 @@ void	MemoryUnmap(void* Address, SIZE_T Size);
 
 ////////////////////////////////////////////////////////////////////////////////
 UPTRINT	ThreadCreate(const ANSICHAR* Name, void (*Entry)());
-uint32	ThreadGetCurrentId();
 void	ThreadSleep(uint32 Milliseconds);
 void	ThreadJoin(UPTRINT Handle);
 void	ThreadDestroy(UPTRINT Handle);
diff --git a/Engine/Source/Runtime/TraceLog/Private/Trace/Unix/UnixTrace.cpp b/Engine/Source/Runtime/TraceLog/Private/Trace/Unix/UnixTrace.cpp
--- a/Engine/Source/Runtime/TraceLog/Private/Trace/Unix/UnixTrace.cpp
+++ b/Engine/Source/Runtime/TraceLog/Private/Trace/Unix/UnixTrace.cpp
@@ -62,15 +62,6 @@ UPTRINT ThreadCreate(const ANSICHAR* Name, void (*Entry)())
 	return static_cast<UPTRINT>(ThreadHandle);
 }
 
-////////////////////////////////////////////////////////////////////////////////
-uint32 ThreadGetCurrentId()
-{
-	// FIXME: suboptimal
-	pid_t ThreadId = static_cast<pid_t>(syscall(SYS_gettid));
-	static_assert(sizeof(pid_t) <= sizeof(uint32), "pid_t is larger than uint32, reconsider implementation of GetCurrentThreadId()");
-	return static_cast<uint32>(ThreadId);
-}
-
 ////////////////////////////////////////////////////////////////////////////////
 void ThreadSleep(uint32 Milliseconds)
 {
