Removed the explicit initialisation of some globals to zero just to be on the safe side of them getting explicitly initialised.

diff --git a/Engine/Source/Runtime/TraceLog/Private/Trace/Writer.cpp b/Engine/Source/Runtime/TraceLog/Private/Trace/Writer.cpp
--- a/Engine/Source/Runtime/TraceLog/Private/Trace/Writer.cpp
+++ b/Engine/Source/Runtime/TraceLog/Private/Trace/Writer.cpp
@@ -316,10 +316,10 @@ enum class EDataState : uint8
 	Partial,			// Passive, but buffers are full so some events are lost
 	Sending,			// Events are being sent to an IO handle
 };
-static FHoldBuffer		GHoldBuffer;
-static UPTRINT			GDataHandle			= 0;
-static EDataState		GDataState;			// = EDataState::Passive;
-UPTRINT					GPendingDataHandle	= 0;
+static FHoldBuffer		GHoldBuffer;		// will init to zero.
+static UPTRINT			GDataHandle;		// = 0
+static EDataState		GDataState;			// = EDataState::Passive
+UPTRINT					GPendingDataHandle;	// = 0
 
 ////////////////////////////////////////////////////////////////////////////////
 static void Writer_ConsumeEvents()
