Updated how event IDs are determined and how the master event list is built

diff --git a/Engine/Source/Runtime/TraceLog/Private/Trace/Writer.cpp b/Engine/Source/Runtime/TraceLog/Private/Trace/Writer.cpp
--- a/Engine/Source/Runtime/TraceLog/Private/Trace/Writer.cpp
+++ b/Engine/Source/Runtime/TraceLog/Private/Trace/Writer.cpp
@@ -921,16 +921,8 @@ bool Writer_Open(const ANSICHAR* Path)
 
 
 ////////////////////////////////////////////////////////////////////////////////
-template <typename Type>
-struct TLateAtomic
-{
-	typedef TTraceAtomic<Type> InnerType;
-	InnerType* operator -> () { return (InnerType*)Buffer; }
-	alignas(InnerType) char Buffer[sizeof(InnerType)];
-};
-
-static TLateAtomic<uint32>		GEventUidCounter;	// = 0;
-static TLateAtomic<FEventDef*>	GHeadEvent;			// = nullptr;
+static UPTRINT volatile		GEventUidCounter;	// = 0;
+static FEventDef* volatile	GHeadEvent;			// = nullptr;
 
 ////////////////////////////////////////////////////////////////////////////////
 enum class EKnownEventUids : uint16
@@ -976,8 +968,17 @@ void Writer_EventCreate(
 	Writer_Initialize();
 
 	// Assign a unique ID for this event
-	uint32 Uid = GEventUidCounter->fetch_add(1, std::memory_order_relaxed);
-	Uid += uint32(EKnownEventUids::User);
+	uint32 Uid;
+	for (;; Writer_Yield())
+	{
+		UPTRINT CurrentUid = UPTRINT(AtomicLoadRelaxed((void* volatile*)&GEventUidCounter));
+		UPTRINT NextUid = CurrentUid + 1;
+		if (AtomicCompareExchangeRelaxed((void* volatile*)&GEventUidCounter, (void*)NextUid, (void*)CurrentUid))
+		{
+			Uid = uint32(CurrentUid) + uint32(EKnownEventUids::User);
+			break;
+		}
+	}
 
 	if (Uid >= uint32(EKnownEventUids::Max))
 	{
@@ -1053,11 +1054,11 @@ void Writer_EventCreate(
 	Writer_EndLog(&(uint8&)Event);
 
 	// Add this new event into the list so we can look them up later.
-	while (true)
+	for (;; Writer_Yield())
 	{
-		FEventDef* HeadEvent = GHeadEvent->load(std::memory_order_relaxed);
+		FEventDef* HeadEvent = AtomicLoadRelaxed(&GHeadEvent);
 		Target->Handle = HeadEvent;
-		if (GHeadEvent->compare_exchange_weak(HeadEvent, Target, std::memory_order_release))
+		if (AtomicCompareExchangeRelease(&GHeadEvent, Target, HeadEvent))
 		{
 			break;
 		}
@@ -1076,7 +1077,7 @@ uint32 Writer_EventToggle(const ANSICHAR* Wildcard, bool bState)
 	{
 		uint32 LoggerHash = Writer_EventGetHash(Wildcard);
 
-		FEventDef* Event = Private::GHeadEvent->load(std::memory_order_relaxed);
+		FEventDef* Event = AtomicLoadAcquire(&GHeadEvent);
 		for (; Event != nullptr; Event = (FEventDef*)(Event->Handle))
 		{
 			if (Event->LoggerHash == LoggerHash)
@@ -1093,7 +1094,7 @@ uint32 Writer_EventToggle(const ANSICHAR* Wildcard, bool bState)
 	uint32 NameHash = Writer_EventGetHash(Dot + 1);
 	uint32 EventHash = Writer_EventGetHash(LoggerHash, NameHash);
 
-	FEventDef* Event = Private::GHeadEvent->load(std::memory_order_relaxed);
+	FEventDef* Event = (FEventDef*)AtomicLoadAcquire(&Private::GHeadEvent);
 	for (; Event != nullptr; Event = (FEventDef*)(Event->Handle))
 	{
 		if (Event->Hash == EventHash)
diff --git a/Engine/Source/Runtime/TraceLog/Public/Trace/Private/Writer.inl b/Engine/Source/Runtime/TraceLog/Public/Trace/Private/Writer.inl
--- a/Engine/Source/Runtime/TraceLog/Public/Trace/Private/Writer.inl
+++ b/Engine/Source/Runtime/TraceLog/Public/Trace/Private/Writer.inl
@@ -4,19 +4,15 @@
 
 #if UE_TRACE_ENABLED
 
-#include <atomic>
 #include "Atomic.h"
 
 namespace Trace
 {
 
-////////////////////////////////////////////////////////////////////////////////
-template <typename Type>
-using TTraceAtomic = std::atomic<Type>;
-
 namespace Private
 {
 
+////////////////////////////////////////////////////////////////////////////////
 #if defined(_MSC_VER)
 	#pragma warning(push)
 	#pragma warning(disable : 4200) // non-standard zero-sized array
