Only send the hold buffer's content if it has content to send

diff --git a/Engine/Source/Runtime/TraceLog/Private/Trace/Writer.cpp b/Engine/Source/Runtime/TraceLog/Private/Trace/Writer.cpp
--- a/Engine/Source/Runtime/TraceLog/Private/Trace/Writer.cpp
+++ b/Engine/Source/Runtime/TraceLog/Private/Trace/Writer.cpp
@@ -478,7 +478,10 @@ static void Writer_UpdateData()
 			bOk &= IoWrite(GDataHandle, &TransportHeader, sizeof(TransportHeader));
 
 			// Passively collected data
-			bOk &= IoWrite(GDataHandle, GHoldBuffer->GetData(), GHoldBuffer->GetSize());
+			if (GHoldBuffer->GetSize())
+			{
+				bOk &= IoWrite(GDataHandle, GHoldBuffer->GetData(), GHoldBuffer->GetSize());
+			}
 
 			if (bOk)
 			{
