Event definition events are pretty important. We should say as such to someone.

diff --git a/Engine/Source/Runtime/TraceLog/Private/Trace/Writer.cpp b/Engine/Source/Runtime/TraceLog/Private/Trace/Writer.cpp
--- a/Engine/Source/Runtime/TraceLog/Private/Trace/Writer.cpp
+++ b/Engine/Source/Runtime/TraceLog/Private/Trace/Writer.cpp
@@ -1016,10 +1016,11 @@ void Writer_EventCreate(
 	}
 
 	// Allocate the new event event in the log stream.
+	uint16 EventUid = uint16(EKnownEventUids::NewEvent)|uint16(EKnownEventUids::Flag_Important);
 	uint16 EventSize = sizeof(FNewEventEvent);
 	EventSize += sizeof(FNewEventEvent::Fields[0]) * FieldCount;
 	EventSize += NamesSize;
-	auto& Event = *(FNewEventEvent*)Writer_BeginLog(uint16(EKnownEventUids::NewEvent), EventSize);
+	auto& Event = *(FNewEventEvent*)Writer_BeginLog(EventUid, EventSize);
 
 	// Write event's main properties.
 	Event.EventUid = uint16(Uid) & uint16(EKnownEventUids::UidMask);
