Fixed Clang-Win64 compile error

diff --git a/Engine/Source/Runtime/TraceLog/Private/Trace/Writer.cpp b/Engine/Source/Runtime/TraceLog/Private/Trace/Writer.cpp
--- a/Engine/Source/Runtime/TraceLog/Private/Trace/Writer.cpp
+++ b/Engine/Source/Runtime/TraceLog/Private/Trace/Writer.cpp
@@ -60,7 +60,7 @@ void Writer_InitializeTiming()
 static uint8						GEmptyBuffer[sizeof(FWriteBuffer)];
 thread_local FWriteBuffer*			GWriteBuffer		= (FWriteBuffer*)GEmptyBuffer;
 T_ALIGN static void* volatile		GFirstEvent;
-UE_TRACE_API T_ALIGN void* volatile	GLastEvent;			// = nullptr;
+T_ALIGN UE_TRACE_API void* volatile	GLastEvent;			// = nullptr;
 static const uint32					GPoolSize			= 384 << 20; // 384MB ought to be enough
 T_ALIGN static UPTRINT volatile		GThreadId;			// = 0;
 static const uint32					GPoolBlockSize		= 4 << 10;
