Moved a misplaced #include directive

diff --git a/Engine/Source/Runtime/TraceLog/Private/Trace/Writer.cpp b/Engine/Source/Runtime/TraceLog/Private/Trace/Writer.cpp
--- a/Engine/Source/Runtime/TraceLog/Private/Trace/Writer.cpp
+++ b/Engine/Source/Runtime/TraceLog/Private/Trace/Writer.cpp
@@ -7,7 +7,9 @@
 #include "Trace/Platform.h"
 #include "Trace/Private/Atomic.h"
 #include "Trace/Trace.h"
+
 #include "Misc/CString.h"
+#include "Templates/UnrealTemplate.h"
 
 #if PLATFORM_CPU_X86_FAMILY
 	#include <emmintrin.h>
diff --git a/Engine/Source/Runtime/TraceLog/Public/Trace/Private/Field.h b/Engine/Source/Runtime/TraceLog/Public/Trace/Private/Field.h
--- a/Engine/Source/Runtime/TraceLog/Public/Trace/Private/Field.h
+++ b/Engine/Source/Runtime/TraceLog/Public/Trace/Private/Field.h
@@ -3,7 +3,6 @@
 #pragma once
 
 #include "Trace.h"
-#include "Templates/UnrealTemplate.h"
 
 namespace Trace
 {
