Removed ability to trace log events into an alternative buffer as this no longer maps to the new TLS-based approach of buffering events

diff --git a/Engine/Source/Runtime/TraceLog/Public/Trace/Private/EventDef.h b/Engine/Source/Runtime/TraceLog/Public/Trace/Private/EventDef.h
--- a/Engine/Source/Runtime/TraceLog/Public/Trace/Private/EventDef.h
+++ b/Engine/Source/Runtime/TraceLog/Public/Trace/Private/EventDef.h
@@ -64,10 +64,8 @@ public:
 	public:
 						FLogScope(uint16 EventUid, uint16 Size);
 						FLogScope(uint16 EventUid, uint16 Size, uint16 ExtraBytes);
-						FLogScope(uint16 EventUid, uint16 Size, uint8* Out);
 						~FLogScope();
 		uint8*			Ptr;
-		const bool		bOutOfBand = false;
 	};
 
 	void*						Handle;
diff --git a/Engine/Source/Runtime/TraceLog/Public/Trace/Private/EventDef.inl b/Engine/Source/Runtime/TraceLog/Public/Trace/Private/EventDef.inl
--- a/Engine/Source/Runtime/TraceLog/Public/Trace/Private/EventDef.inl
+++ b/Engine/Source/Runtime/TraceLog/Public/Trace/Private/EventDef.inl
@@ -22,21 +22,10 @@ inline FEventDef::FLogScope::FLogScope(uint16 EventUid, uint16 Size, uint16 Extr
 	Ptr = Writer_BeginLog(EventUid, Size + ExtraBytes);
 }
 
-////////////////////////////////////////////////////////////////////////////////
-inline FEventDef::FLogScope::FLogScope(uint16 EventUid, uint16 Size, uint8* Out)
-: bOutOfBand(true)
-{
-	*(uint32*)Out = (uint32(Size) << 16)|uint32(EventUid);
-	Ptr = Out + HeaderSize;
-}
-
 ////////////////////////////////////////////////////////////////////////////////
 inline FEventDef::FLogScope::~FLogScope()
 {
-	if (!bOutOfBand)
-	{
-		Writer_EndLog(Ptr);
-	}
+	Writer_EndLog(Ptr);
 }
 
 
