Fixed a compiler error that was entirely of my own doing.

diff --git a/Engine/Source/Runtime/TraceLog/Public/Trace/Private/Field.h b/Engine/Source/Runtime/TraceLog/Public/Trace/Private/Field.h
--- a/Engine/Source/Runtime/TraceLog/Public/Trace/Private/Field.h
+++ b/Engine/Source/Runtime/TraceLog/Public/Trace/Private/Field.h
@@ -4,6 +4,8 @@
 
 #include "Trace.h"
 
+#include "Templates/UnrealTemplate.h"
+
 namespace Trace
 {
 
