As TraceLog will rely on defines/inlines from Core in its public headers, the include-only dependency should be declared as a public one, not a private one.

diff --git a/Engine/Source/Runtime/TraceLog/TraceLog.Build.cs b/Engine/Source/Runtime/TraceLog/TraceLog.Build.cs
--- a/Engine/Source/Runtime/TraceLog/TraceLog.Build.cs
+++ b/Engine/Source/Runtime/TraceLog/TraceLog.Build.cs
@@ -9,7 +9,7 @@ public class TraceLog : ModuleRules
 	public TraceLog(ReadOnlyTargetRules Target) : base(Target)
 	{
 		bRequiresImplementModule = false;
-		PrivateIncludePathModuleNames.Add("Core");
+		PublicIncludePathModuleNames.Add("Core");
 
         if (Target.Platform == UnrealTargetPlatform.PS4)
         {
