Moved TraceLog's internal-but-public header files from the confusingly named "Private/" folder to a "Detail/" folder.

diff --git a/Engine/Source/Developer/TraceAnalysis/Private/Analysis/Engine.cpp b/Engine/Source/Developer/TraceAnalysis/Private/Analysis/Engine.cpp
--- a/Engine/Source/Developer/TraceAnalysis/Private/Analysis/Engine.cpp
+++ b/Engine/Source/Developer/TraceAnalysis/Private/Analysis/Engine.cpp
@@ -6,7 +6,7 @@
 #include "HAL/UnrealMemory.h"
 #include "Trace/Analysis.h"
 #include "Trace/Analyzer.h"
-#include "Trace/Private/EventDef.h"
+#include "Trace/Detail/EventDef.h"
 #include "Transport/Transport.h"
 #include "Transport/TlsTransport.h"
 
diff --git a/Engine/Source/Developer/TraceAnalysis/Public/Trace/Analyzer.h b/Engine/Source/Developer/TraceAnalysis/Public/Trace/Analyzer.h
--- a/Engine/Source/Developer/TraceAnalysis/Public/Trace/Analyzer.h
+++ b/Engine/Source/Developer/TraceAnalysis/Public/Trace/Analyzer.h
@@ -4,7 +4,7 @@
 
 #include "CoreTypes.h"
 #include "Logging/LogMacros.h"
-#include "Trace/Private/Field.h"
+#include "Trace/Detail/Field.h"
 
 #include <memory.h>
 
diff --git a/Engine/Source/Programs/UnrealBuildTool/System/Unity.cs b/Engine/Source/Programs/UnrealBuildTool/System/Unity.cs
--- a/Engine/Source/Programs/UnrealBuildTool/System/Unity.cs
+++ b/Engine/Source/Programs/UnrealBuildTool/System/Unity.cs
@@ -231,7 +231,7 @@ namespace UnrealBuildTool
 					{
 						// Every single file in the module appears in the working set, so don't bother using adaptive unity for this
 						// module.  Otherwise it would make full builds really slow.
-						bUseAdaptiveUnityBuild = false;
+						//bUseAdaptiveUnityBuild = false;
 					}
 				}
 
diff --git a/Engine/Source/Runtime/TraceLog/Private/Trace/EventDef.cpp b/Engine/Source/Runtime/TraceLog/Private/Trace/EventDef.cpp
--- a/Engine/Source/Runtime/TraceLog/Private/Trace/EventDef.cpp
+++ b/Engine/Source/Runtime/TraceLog/Private/Trace/EventDef.cpp
@@ -1,11 +1,11 @@
 // Copyright 1998-2019 Epic Games, Inc. All Rights Reserved.
 
-#include "Trace/Private/EventDef.h"
+#include "Trace/Detail/EventDef.h"
 
 #if UE_TRACE_ENABLED
 
-#include "Trace/Private/Field.h"
-#include "Trace/Private/Writer.inl"
+#include "Trace/Detail/Field.h"
+#include "Trace/Detail/Writer.inl"
 
 namespace Trace
 {
diff --git a/Engine/Source/Runtime/TraceLog/Private/Trace/Platform.h b/Engine/Source/Runtime/TraceLog/Private/Trace/Platform.h
--- a/Engine/Source/Runtime/TraceLog/Private/Trace/Platform.h
+++ b/Engine/Source/Runtime/TraceLog/Private/Trace/Platform.h
@@ -2,7 +2,7 @@
 
 #pragma once
 
-#include "Trace/Private/Trace.h"
+#include "Trace/Detail/Trace.h"
 
 #if UE_TRACE_ENABLED
 
diff --git a/Engine/Source/Runtime/TraceLog/Private/Trace/Trace.cpp b/Engine/Source/Runtime/TraceLog/Private/Trace/Trace.cpp
--- a/Engine/Source/Runtime/TraceLog/Private/Trace/Trace.cpp
+++ b/Engine/Source/Runtime/TraceLog/Private/Trace/Trace.cpp
@@ -1,7 +1,7 @@
 // Copyright 1998-2019 Epic Games, Inc. All Rights Reserved.
 
 #include "Trace/Trace.h"
-#include "Trace/Private/EventDef.h"
+#include "Trace/Detail/EventDef.h"
 
 #if UE_TRACE_ENABLED
 
diff --git a/Engine/Source/Runtime/TraceLog/Private/Trace/Writer.cpp b/Engine/Source/Runtime/TraceLog/Private/Trace/Writer.cpp
--- a/Engine/Source/Runtime/TraceLog/Private/Trace/Writer.cpp
+++ b/Engine/Source/Runtime/TraceLog/Private/Trace/Writer.cpp
@@ -1,11 +1,11 @@
 // Copyright 1998-2019 Epic Games, Inc. All Rights Reserved.
 
-#include "Trace/Private/Trace.h"
+#include "Trace/Detail/Trace.h"
 
 #if UE_TRACE_ENABLED
 
 #include "Trace/Platform.h"
-#include "Trace/Private/Atomic.h"
+#include "Trace/Detail/Atomic.h"
 #include "Trace/Trace.h"
 
 #include "Misc/CString.h"
diff --git a/Engine/Source/Runtime/TraceLog/Public/Trace/Detail/Atomic.h b/Engine/Source/Runtime/TraceLog/Public/Trace/Detail/Atomic.h
new file mode 100644
--- /dev/null
+++ b/Engine/Source/Runtime/TraceLog/Public/Trace/Detail/Atomic.h
@@ -0,0 +1,182 @@
+// Copyright 1998-2019 Epic Games, Inc. All Rights Reserved.
+
+#pragma once
+
+namespace Trace {
+namespace Private {
+
+////////////////////////////////////////////////////////////////////////////////
+template <typename Type> Type	AtomicLoadRelaxed(Type volatile* Source);
+template <typename Type> Type	AtomicLoadAcquire(Type volatile* Source);
+template <typename Type> void	AtomicStoreRelaxed(Type volatile* Target, Type Value);
+template <typename Type> void	AtomicStoreRelease(Type volatile* Target, Type Value);
+template <typename Type> bool	AtomicCompareExchangeRelaxed(Type volatile* Target, Type New, Type Expected);
+template <typename Type> bool	AtomicCompareExchangeAcquire(Type volatile* Target, Type New, Type Expected);
+template <typename Type> bool	AtomicCompareExchangeRelease(Type volatile* Target, Type New, Type Expected);
+
+} // namespace Private
+} // namespace Trace
+
+////////////////////////////////////////////////////////////////////////////////
+#define IS_MSVC					0
+#define IS_GCC_COMPATIBLE		0
+#if defined(_MSC_VER) && !defined(__clang__)
+#	undef  IS_MSVC
+#	define IS_MSVC				1
+#elif defined(__clang__) || defined(__GNUC__)
+#	undef  IS_GCC_COMPATIBLE
+#	define IS_GCC_COMPATIBLE	1
+#endif
+
+////////////////////////////////////////////////////////////////////////////////
+#if IS_MSVC
+#	include <intrin.h>
+#	if defined(_M_ARM) || defined(_M_ARM64)
+#		if defined(_M_ARM)
+#			include <armintr.h>
+#		else
+#			include <arm64intr.h>
+#		endif
+#	endif
+#endif
+
+namespace Trace {
+namespace Private {
+
+////////////////////////////////////////////////////////////////////////////////
+#if IS_MSVC
+#	if defined(_M_ARM) || defined(_M_ARM64)
+#		define INTERLOCKED_API(Name, Suffix, ...)	Name##_##Suffix(__VA_ARGS__)
+#	elif defined(_M_IX86) || defined(_M_X64)
+#		define INTERLOCKED_API(Name, Suffix, ...)	Name(__VA_ARGS__)
+#	endif
+#endif
+
+////////////////////////////////////////////////////////////////////////////////
+template <typename Type>
+inline Type AtomicLoadRelaxed(Type volatile* Source)
+{
+	static_assert(sizeof(Type) == sizeof(void*), "");
+#if IS_MSVC
+#	if defined(_M_ARM)
+		__int32 Value = __iso_volatile_load32((__int32 volatile*)Source);
+		return (Type*)Value;
+#	elif defined(_M_ARM64)
+		__int64 Value = __iso_volatile_load64((__int64 volatile*)Source);
+		return (Type*)Value;
+#	else
+		return *Source;
+#	endif
+#elif IS_GCC_COMPATIBLE
+	return __atomic_load_n(Source, __ATOMIC_RELAXED);
+#endif
+}
+
+////////////////////////////////////////////////////////////////////////////////
+template <typename Type>
+inline Type AtomicLoadAcquire(Type volatile* Source)
+{
+	static_assert(sizeof(Type) == sizeof(void*), "");
+#if IS_MSVC
+#	if defined(_M_ARM)
+		__int32 Value = __iso_volatile_load32((__int32 volatile*)Source);
+		__dmb(_ARM_BARRIER_ISH);
+		return Type(Value);
+#	elif defined(_M_ARM64)
+		__int64 Value = __iso_volatile_load64((__int64 volatile*)Source);
+		__dmb(_ARM64_BARRIER_ISH);
+		return Type(Value);
+#	else
+		Type Value = *Source;
+		_ReadWriteBarrier();
+		return Value;
+#	endif
+#elif IS_GCC_COMPATIBLE
+	return __atomic_load_n(Source, __ATOMIC_ACQUIRE);
+#endif
+}
+
+////////////////////////////////////////////////////////////////////////////////
+template <typename Type>
+inline void AtomicStoreRelaxed(Type volatile* Target, Type Value)
+{
+	static_assert(sizeof(Type) == sizeof(void*), "");
+#if IS_MSVC
+#	if defined(_M_ARM)
+		__iso_volatile_store32((__int32 volatile*)Target, UPTRINT(Value));
+#	elif defined(_M_ARM64)
+		__iso_volatile_store64((__int64 volatile*)Target, UPTRINT(Value));
+#	else
+		*Target = Value;
+#	endif
+#elif IS_GCC_COMPATIBLE
+	return __atomic_store_n(Target, Value, __ATOMIC_RELAXED);
+#endif
+}
+
+////////////////////////////////////////////////////////////////////////////////
+template <typename Type>
+inline void AtomicStoreRelease(Type volatile* Target, Type Value)
+{
+	static_assert(sizeof(Type) == sizeof(void*), "");
+#if IS_MSVC
+#	if defined(_M_ARM)
+		__dmb(_ARM_BARRIER_ISH);
+		__iso_volatile_store32((__int32 volatile*)Target, __int32(Value));
+#	elif defined(_M_ARM64)
+		__dmb(_ARM64_BARRIER_ISH);
+		__iso_volatile_store64((__int64 volatile*)Target, __int64(Value));
+#	else
+		_ReadWriteBarrier();
+		*Target = Value;
+#	endif
+#elif IS_GCC_COMPATIBLE
+	return __atomic_store_n(Target, Value, __ATOMIC_RELEASE);
+#endif
+}
+
+////////////////////////////////////////////////////////////////////////////////
+template <typename Type>
+inline bool AtomicCompareExchangeRelaxed(Type volatile* Target, Type New, Type Expected)
+{
+	static_assert(sizeof(Type) == sizeof(void*), "");
+#if IS_MSVC
+	return INTERLOCKED_API(_InterlockedCompareExchangePointer, _nf, (void* volatile*)Target, New, Expected) == Expected;
+#elif IS_GCC_COMPATIBLE
+	Type InOut = Expected;
+	return __atomic_compare_exchange_n(Target, &InOut, New, true, __ATOMIC_RELAXED, __ATOMIC_RELAXED);
+#endif
+}
+
+////////////////////////////////////////////////////////////////////////////////
+template <typename Type>
+inline bool AtomicCompareExchangeAcquire(Type volatile* Target, Type New, Type Expected)
+{
+	static_assert(sizeof(Type) == sizeof(void*), "");
+#if IS_MSVC
+	return INTERLOCKED_API(_InterlockedCompareExchangePointer, _acq, (void* volatile*)Target, New, Expected) == Expected;
+#elif IS_GCC_COMPATIBLE
+	Type InOut = Expected;
+	return __atomic_compare_exchange_n(Target, &InOut, New, true, __ATOMIC_ACQUIRE, __ATOMIC_RELAXED);
+#endif
+}
+
+////////////////////////////////////////////////////////////////////////////////
+template <typename Type>
+inline bool AtomicCompareExchangeRelease(Type volatile* Target, Type New, Type Expected)
+{
+	static_assert(sizeof(Type) == sizeof(void*), "");
+#if IS_MSVC
+	return INTERLOCKED_API(_InterlockedCompareExchangePointer, _rel, (void* volatile*)Target, New, Expected) == Expected;
+#elif IS_GCC_COMPATIBLE
+	Type InOut = Expected;
+	return __atomic_compare_exchange_n(Target, &InOut, New, true, __ATOMIC_RELEASE, __ATOMIC_RELAXED);
+#endif
+}
+
+#undef INTERLOCKED_API
+#undef IS_GCC_COMPATIBLE
+#undef IS_MSVC
+
+} // namespace Private
+} // namespace Trace
diff --git a/Engine/Source/Runtime/TraceLog/Public/Trace/Detail/EventDef.h b/Engine/Source/Runtime/TraceLog/Public/Trace/Detail/EventDef.h
new file mode 100644
--- /dev/null
+++ b/Engine/Source/Runtime/TraceLog/Public/Trace/Detail/EventDef.h
@@ -0,0 +1,90 @@
+// Copyright 1998-2019 Epic Games, Inc. All Rights Reserved.
+
+#pragma once
+
+#include "Trace.h"
+
+namespace Trace
+{
+
+#if defined(_MSC_VER)
+	#pragma warning(push)
+	#pragma warning(disable : 4200) // non-standard zero-sized array
+#endif
+
+////////////////////////////////////////////////////////////////////////////////
+struct FNewEventEvent
+{
+	enum : uint16 { Uid = 0 };
+
+	uint16		EventUid;
+	uint16		FieldCount;
+	uint8		LoggerNameSize;
+	uint8		EventNameSize;
+	struct
+	{
+		uint16	Offset;
+		uint16	Size;
+		uint8	TypeInfo;
+		uint8	NameSize;
+	}			Fields[];
+	/*uint8		NameData[]*/
+};
+
+#if defined(_MSC_VER)
+	#pragma warning(pop)
+#endif
+
+} // namespace Trace
+
+
+
+#if UE_TRACE_ENABLED
+
+namespace Trace
+{
+
+struct FFieldDesc;
+struct FLiteralName;
+
+////////////////////////////////////////////////////////////////////////////////
+class FEventDef
+{
+public:
+	enum { HeaderSize = sizeof(uint32), }; // I'll just dump this here for now.
+
+	enum
+	{
+		Flag_Always		= 1 << 0,
+		Flag_Important	= 1 << 1,
+	};
+
+	class FLogScope
+	{
+	public:
+						FLogScope(uint16 EventUid, uint16 Size);
+						FLogScope(uint16 EventUid, uint16 Size, uint16 ExtraBytes);
+						~FLogScope();
+		uint8*			Ptr;
+	};
+
+	void*						Handle;
+	uint32						LoggerHash;
+	uint32						Hash;
+	uint16						Uid;
+	union
+	{
+		struct
+		{
+			bool				bOptedIn;
+			uint8				Internal;
+		};
+		uint16					Test;
+	}							Enabled;
+	bool						bInitialized;
+	UE_TRACE_API static void	Create(FEventDef* Target, const FLiteralName& LoggerName, const FLiteralName& EventName, const FFieldDesc* FieldDescs, uint32 FieldCount, uint32 Flags=0);
+};
+
+} // namespace Trace
+
+#endif // UE_TRACE_ENABLED
diff --git a/Engine/Source/Runtime/TraceLog/Public/Trace/Detail/EventDef.inl b/Engine/Source/Runtime/TraceLog/Public/Trace/Detail/EventDef.inl
new file mode 100644
--- /dev/null
+++ b/Engine/Source/Runtime/TraceLog/Public/Trace/Detail/EventDef.inl
@@ -0,0 +1,43 @@
+// Copyright 1998-2019 Epic Games, Inc. All Rights Reserved.
+
+#pragma once
+
+#if UE_TRACE_ENABLED
+
+#include "EventDef.h"
+#include "Writer.inl"
+
+namespace Trace
+{
+
+////////////////////////////////////////////////////////////////////////////////
+inline FEventDef::FLogScope::FLogScope(uint16 EventUid, uint16 Size)
+{
+	Ptr = Writer_BeginLog(EventUid, Size);
+}
+
+////////////////////////////////////////////////////////////////////////////////
+inline FEventDef::FLogScope::FLogScope(uint16 EventUid, uint16 Size, uint16 ExtraBytes)
+{
+	Ptr = Writer_BeginLog(EventUid, Size + ExtraBytes);
+}
+
+////////////////////////////////////////////////////////////////////////////////
+inline FEventDef::FLogScope::~FLogScope()
+{
+	Writer_EndLog(Ptr);
+}
+
+
+
+////////////////////////////////////////////////////////////////////////////////
+template <typename ActionType>
+inline const FEventDef::FLogScope& operator << (const FEventDef::FLogScope& Lhs, const ActionType& Rhs)
+{
+	Rhs.Write(Lhs.Ptr);
+	return Lhs;
+}
+
+} // namespace Trace
+
+#endif // UE_TRACE_ENABLED
diff --git a/Engine/Source/Runtime/TraceLog/Public/Trace/Detail/Field.h b/Engine/Source/Runtime/TraceLog/Public/Trace/Detail/Field.h
new file mode 100644
--- /dev/null
+++ b/Engine/Source/Runtime/TraceLog/Public/Trace/Detail/Field.h
@@ -0,0 +1,220 @@
+// Copyright 1998-2019 Epic Games, Inc. All Rights Reserved.
+
+#pragma once
+
+#include "Trace.h"
+
+#include "Templates/UnrealTemplate.h"
+
+namespace Trace
+{
+
+////////////////////////////////////////////////////////////////////////////////
+enum : uint8
+{
+	/* Category */
+	_Field_CategoryMask	= 0300,
+	_Field_Integer		= 0000, 
+	_Field_Float		= 0100,
+	_Field_Array		= 0200, /* flag */
+
+	/* Size */
+	_Field_Pow2SizeMask	= 0003,
+	_Field_8			= 0000,
+	_Field_16			= 0001,
+	_Field_32			= 0002,
+	_Field_64			= 0003,
+#if PLATFORM_64BITS
+	_Field_Ptr			= _Field_64,
+#else
+	_Field_Ptr			= _Field_32,
+#endif
+
+	/* unused			= 0004, flag */
+	/* unused			= 0070, enum/flags*/
+};
+
+////////////////////////////////////////////////////////////////////////////////
+enum EFieldType : uint8
+{
+	Field_Bool		= _Field_Integer | _Field_8,
+	Field_Int8		= _Field_Integer | _Field_8,
+	Field_Int16		= _Field_Integer | _Field_16,
+	Field_Int32		= _Field_Integer | _Field_32,
+	Field_Int64		= _Field_Integer | _Field_64,
+	Field_Ptr		= _Field_Integer | _Field_Ptr,
+	Field_Float		= _Field_Float	 | _Field_32,
+	Field_Double	= _Field_Float	 | _Field_64,
+
+	/* Flags */
+	Field_Array		= _Field_Array,
+};
+
+////////////////////////////////////////////////////////////////////////////////
+template <typename Type> struct TFieldType;
+
+template <> struct TFieldType<bool>			{ enum { Value = Field_Bool }; };
+template <> struct TFieldType<int8>			{ enum { Value = Field_Int8 }; };
+template <> struct TFieldType<int16>		{ enum { Value = Field_Int16 }; };
+template <> struct TFieldType<int32>		{ enum { Value = Field_Int32 }; };
+template <> struct TFieldType<int64>		{ enum { Value = Field_Int64 }; };
+template <> struct TFieldType<uint8>		{ enum { Value = Field_Int8 }; };
+template <> struct TFieldType<uint16>		{ enum { Value = Field_Int16 }; };
+template <> struct TFieldType<uint32>		{ enum { Value = Field_Int32 }; };
+template <> struct TFieldType<uint64>		{ enum { Value = Field_Int64 }; };
+template <> struct TFieldType<float>		{ enum { Value = Field_Float }; };
+template <> struct TFieldType<double>		{ enum { Value = Field_Double }; };
+template <class T> struct TFieldType<T*>	{ enum { Value = Field_Ptr }; };
+
+} // namespace Trace
+
+#if UE_TRACE_ENABLED
+
+namespace Trace
+{
+
+////////////////////////////////////////////////////////////////////////////////
+struct FLiteralName
+{
+	template <uint32 Size>
+	explicit FLiteralName(const ANSICHAR (&Name)[Size])
+	: Ptr(Name)
+	, Length(Size - 1)
+	{
+		static_assert(Size < 256, "Field name is too large");
+	}
+
+	const ANSICHAR* Ptr;
+	uint8 Length;
+};
+
+////////////////////////////////////////////////////////////////////////////////
+struct FFieldDesc
+{
+	FFieldDesc(const FLiteralName& Name, uint8 Type, uint16 Offset, uint16 Size)
+	: Name(Name.Ptr)
+	, ValueOffset(Offset)
+	, ValueSize(Size)
+	, NameSize(Name.Length)
+	, TypeInfo(Type)
+	{
+	}
+
+	const ANSICHAR* Name;
+	uint16			ValueOffset;
+	uint16			ValueSize;
+	uint8			NameSize;
+	uint8			TypeInfo;
+};
+
+////////////////////////////////////////////////////////////////////////////////
+template <int OffsetValue, typename Type>
+struct TFieldBase
+	: public FFieldDesc
+{
+	enum { Offset = OffsetValue, Size = sizeof(Type) };
+
+protected:
+	TFieldBase(const FLiteralName& Name)
+	: FFieldDesc(Name, TFieldType<Type>::Value, OffsetValue, Size)
+	{
+	}
+};
+
+////////////////////////////////////////////////////////////////////////////////
+template <int Offset, typename Type> struct TField;
+
+////////////////////////////////////////////////////////////////////////////////
+// Used to terminate the field list and determine an event's size.
+enum EndOfFields {};
+template <int Size>
+struct TField<Size, EndOfFields>
+{
+	enum : uint16 { Value = Size };
+};
+
+////////////////////////////////////////////////////////////////////////////////
+// Access to additional data that can be included along with a logged event.
+enum Attachment {};
+template <int Offset>
+struct TField<Offset, Attachment>
+{
+	template <typename LambdaType>
+	struct FActionableLambda
+	{
+		LambdaType& Value;
+		void Write(uint8* __restrict Ptr) const
+		{
+			Value(Ptr + Offset);
+		}
+	};
+
+	template <typename LambdaType>
+	const FActionableLambda<LambdaType> operator () (LambdaType&& Lambda) const
+	{
+		return { Forward<LambdaType>(Lambda) };
+	}
+
+
+	struct FActionableMemcpy
+	{
+		const void* Data;
+		uint32 Size;
+		void Write(uint8* __restrict Ptr) const
+		{
+			::memcpy(Ptr + Offset, Data, Size);
+		}
+	};
+
+	const FActionableMemcpy operator () (const void* Data, uint32 Size) const
+	{
+		return { Data, Size };
+	}
+};
+
+#if 0
+////////////////////////////////////////////////////////////////////////////////
+template <typename Type>
+struct TField<Type[]>
+	: public TFieldBase<sizeof(uint32)>
+{
+	TField(const ANSICHAR (Name)[]) : TFieldBase<sizeof(uint32)>(Name) {}
+};
+
+////////////////////////////////////////////////////////////////////////////////
+template <int Offset, typename Type, int Count>
+struct TField<Offset, Type[Count]>
+	: public TFieldBase<Offset, Type>
+{
+	TField(const ANSICHAR (Name)[]) : TFieldBase<Offset, Type[Count]>(Name) {}
+};
+#endif
+
+////////////////////////////////////////////////////////////////////////////////
+template <int Offset, typename Type>
+struct TField
+	: public TFieldBase<Offset, Type>
+{
+	struct FActionable
+	{
+		Type Value;
+		void Write(uint8* __restrict Ptr) const
+		{
+			*(Type*)(Ptr + Offset) = Value;
+		}
+	};
+
+	TField(const FLiteralName& Name)
+	: TFieldBase<Offset, Type>(Name)
+	{
+	}
+
+	const FActionable operator () (const Type& __restrict Value) const
+	{
+		return { Value };
+	}
+};
+
+} // namespace Trace
+
+#endif // UE_TRACE_ENABLED
diff --git a/Engine/Source/Runtime/TraceLog/Public/Trace/Detail/Trace.h b/Engine/Source/Runtime/TraceLog/Public/Trace/Detail/Trace.h
new file mode 100644
--- /dev/null
+++ b/Engine/Source/Runtime/TraceLog/Public/Trace/Detail/Trace.h
@@ -0,0 +1,124 @@
+// Copyright 1998-2019 Epic Games, Inc. All Rights Reserved.
+
+#pragma once
+
+#include "CoreTypes.h"
+
+////////////////////////////////////////////////////////////////////////////////
+#if !defined(UE_TRACE_ENABLED)
+#	if !UE_BUILD_SHIPPING && !IS_PROGRAM
+#		if PLATFORM_WINDOWS || PLATFORM_PS4 || PLATFORM_XBOXONE || PLATFORM_SWITCH
+#			define UE_TRACE_ENABLED	1
+#		endif
+#	endif
+#endif // !IS_PROGRAM
+
+#if !defined(UE_TRACE_ENABLED)
+#	define UE_TRACE_ENABLED		0
+#endif
+
+#if UE_TRACE_ENABLED
+#	define UE_TRACE_IMPL(...)
+#	define UE_TRACE_API			TRACELOG_API
+#else
+#	define UE_TRACE_IMPL(...)	{ return __VA_ARGS__; }
+#	define UE_TRACE_API			inline
+#endif
+
+
+
+#if UE_TRACE_ENABLED
+
+#define TRACE_PRIVATE_EVENT_DEFINE(LoggerName, EventName) \
+	Trace::FEventDef LoggerName##EventName##Event;
+
+#define TRACE_PRIVATE_EVENT_BEGIN(LoggerName, EventName, ...) \
+	TRACE_PRIVATE_EVENT_BEGIN_IMPL(static, LoggerName, EventName, ##__VA_ARGS__)
+
+#define TRACE_PRIVATE_EVENT_BEGIN_EXTERN(LoggerName, EventName, ...) \
+	TRACE_PRIVATE_EVENT_BEGIN_IMPL(extern, LoggerName, EventName, ##__VA_ARGS__)
+
+#define TRACE_PRIVATE_EVENT_BEGIN_IMPL(LinkageType, LoggerName, EventName, ...) \
+	LinkageType TRACE_PRIVATE_EVENT_DEFINE(LoggerName, EventName) \
+	struct F##LoggerName##EventName##Fields \
+	{ \
+		static void FORCENOINLINE Initialize() \
+		{ \
+			static const bool bOnceOnly = [] () \
+			{ \
+				const uint32 Always = Trace::FEventDef::Flag_Always; \
+				const uint32 Important = Trace::FEventDef::Flag_Important; \
+				const uint32 Flags = (0, ##__VA_ARGS__); \
+				F##LoggerName##EventName##Fields Fields; \
+				const auto* Descs = (Trace::FFieldDesc*)(&Fields); \
+				uint32 DescCount = uint32(sizeof(Fields) / sizeof(*Descs)); \
+				const auto& LoggerLiteral = Trace::FLiteralName(#LoggerName); \
+				const auto& EventLiteral = Trace::FLiteralName(#EventName); \
+				Trace::FEventDef::Create(&LoggerName##EventName##Event, LoggerLiteral, EventLiteral, Descs, DescCount, Flags); \
+				return true; \
+			}(); \
+		} \
+		Trace::TField<0, 
+
+#define TRACE_PRIVATE_EVENT_FIELD(FieldType, FieldName) \
+		FieldType> const FieldName = Trace::FLiteralName(#FieldName); \
+		Trace::TField<decltype(FieldName)::Offset + decltype(FieldName)::Size,
+
+#define TRACE_PRIVATE_EVENT_END() \
+		Trace::EndOfFields> const EventSize_Private = {}; \
+		Trace::TField<decltype(EventSize_Private)::Value, Trace::Attachment> const Attachment = {}; \
+		explicit operator bool () const { return true; } \
+	};
+
+#define TRACE_PRIVATE_EVENT_IS_ENABLED(LoggerName, EventName) \
+	( \
+		(LoggerName##EventName##Event.bInitialized || (F##LoggerName##EventName##Fields::Initialize(), true)) \
+		&& (LoggerName##EventName##Event.Enabled.Test) \
+	)
+
+#define TRACE_PRIVATE_EVENT_SIZE(LoggerName, EventName) \
+	decltype(F##LoggerName##EventName##Fields::EventSize_Private)::Value
+
+#define TRACE_PRIVATE_LOG(LoggerName, EventName, ...) \
+	if (TRACE_PRIVATE_EVENT_IS_ENABLED(LoggerName, EventName)) \
+		if (const auto& __restrict EventName = (F##LoggerName##EventName##Fields&)LoggerName##EventName##Event) \
+			Trace::FEventDef::FLogScope(LoggerName##EventName##Event.Uid, TRACE_PRIVATE_EVENT_SIZE(LoggerName, EventName), ##__VA_ARGS__)
+
+#else
+
+#define TRACE_PRIVATE_EVENT_DEFINE(LoggerName, EventName)
+
+#define TRACE_PRIVATE_EVENT_BEGIN(LoggerName, EventName, ...) \
+	TRACE_PRIVATE_EVENT_BEGIN_IMPL(LoggerName, EventName)
+
+#define TRACE_PRIVATE_EVENT_BEGIN_EXTERN(LoggerName, EventName) \
+	TRACE_PRIVATE_EVENT_BEGIN_IMPL(LoggerName, EventName)
+
+#define TRACE_PRIVATE_EVENT_BEGIN_IMPL(LoggerName, EventName) \
+	struct F##LoggerName##EventName##Dummy \
+	{ \
+		struct FTraceDisabled \
+		{ \
+			const FTraceDisabled& operator () (...) const { return *this; } \
+		}; \
+		const F##LoggerName##EventName##Dummy& operator << (const FTraceDisabled&) const \
+		{ \
+			return *this; \
+		} \
+		explicit operator bool () const { return false; }
+
+#define TRACE_PRIVATE_EVENT_FIELD(FieldType, FieldName) \
+		const FTraceDisabled& FieldName;
+
+#define TRACE_PRIVATE_EVENT_END() \
+		const FTraceDisabled& Attachment; \
+	};
+
+#define TRACE_PRIVATE_EVENT_IS_ENABLED(LoggerName, EventName) \
+	(false)
+
+#define TRACE_PRIVATE_LOG(LoggerName, EventName, ...) \
+	if (const auto& EventName = *(F##LoggerName##EventName##Dummy*)1) \
+		EventName
+
+#endif // UE_TRACE_ENABLED
diff --git a/Engine/Source/Runtime/TraceLog/Public/Trace/Detail/Writer.inl b/Engine/Source/Runtime/TraceLog/Public/Trace/Detail/Writer.inl
new file mode 100644
--- /dev/null
+++ b/Engine/Source/Runtime/TraceLog/Public/Trace/Detail/Writer.inl
@@ -0,0 +1,93 @@
+// Copyright 1998-2019 Epic Games, Inc. All Rights Reserved.
+
+#pragma once
+
+#if UE_TRACE_ENABLED
+
+#include "Atomic.h"
+
+namespace Trace
+{
+
+namespace Private
+{
+
+////////////////////////////////////////////////////////////////////////////////
+#if defined(_MSC_VER)
+	#pragma warning(push)
+	#pragma warning(disable : 4200) // non-standard zero-sized array
+#endif
+
+////////////////////////////////////////////////////////////////////////////////
+struct FWriteBuffer
+{
+	union
+	{
+		uint8*			Cursor;
+		FWriteBuffer*	Next;
+	};
+	uint32				ThreadId;
+	uint8				Data[];
+};
+
+#if defined(_MSC_VER)
+	#pragma warning(pop)
+#endif
+
+////////////////////////////////////////////////////////////////////////////////
+extern UE_TRACE_API void* volatile	GLastEvent;
+UE_TRACE_API uint8*					Writer_NextBuffer(uint16);
+UE_TRACE_API FWriteBuffer*			Writer_GetBuffer();
+
+#if IS_MONOLITHIC
+extern thread_local FWriteBuffer* GWriteBuffer;
+inline FWriteBuffer* Writer_GetBuffer()
+{
+	return GWriteBuffer;
+}
+#endif
+
+} // Private
+
+////////////////////////////////////////////////////////////////////////////////
+inline uint8* Writer_BeginLog(uint16 EventUid, uint16 Size)
+{
+	using namespace Private;
+
+	static const uint32 HeaderSize = sizeof(void*) + sizeof(uint32);
+	uint32 AllocSize = ((Size + HeaderSize) + 7) & ~7;
+
+	FWriteBuffer* Buffer = Writer_GetBuffer();
+	uint8* Cursor = (Buffer->Cursor -= AllocSize);
+	if (UNLIKELY(PTRINT(Cursor) < PTRINT(Buffer->Data)))
+	{
+		Cursor = Writer_NextBuffer(AllocSize);
+	}
+
+	uint32* Out = (uint32*)(Cursor + sizeof(void*));
+	Out[0] = (uint32(Size) << 16)|uint32(EventUid);
+	return (uint8*)(Out + 1);
+}
+
+////////////////////////////////////////////////////////////////////////////////
+inline void Writer_EndLog(uint8* EventData)
+{
+	using namespace Private;
+
+	EventData -= sizeof(void*) + sizeof(uint32);
+
+	// Add the event into the master linked list of events.
+	while (true)
+	{
+		void* Expected = AtomicLoadRelaxed(&GLastEvent);
+		*(void**)EventData = Expected;
+		if (AtomicCompareExchangeRelease(&GLastEvent, (void*)EventData, Expected))
+		{
+			break;
+		}
+	}
+}
+
+} // namespace Trace
+
+#endif // UE_TRACE_ENABLED
diff --git a/Engine/Source/Runtime/TraceLog/Public/Trace/Private/Atomic.h b/Engine/Source/Runtime/TraceLog/Public/Trace/Private/Atomic.h
deleted file mode 100644
diff --git a/Engine/Source/Runtime/TraceLog/Public/Trace/Private/EventDef.h b/Engine/Source/Runtime/TraceLog/Public/Trace/Private/EventDef.h
deleted file mode 100644
diff --git a/Engine/Source/Runtime/TraceLog/Public/Trace/Private/EventDef.inl b/Engine/Source/Runtime/TraceLog/Public/Trace/Private/EventDef.inl
deleted file mode 100644
diff --git a/Engine/Source/Runtime/TraceLog/Public/Trace/Private/Field.h b/Engine/Source/Runtime/TraceLog/Public/Trace/Private/Field.h
deleted file mode 100644
diff --git a/Engine/Source/Runtime/TraceLog/Public/Trace/Private/Trace.h b/Engine/Source/Runtime/TraceLog/Public/Trace/Private/Trace.h
deleted file mode 100644
diff --git a/Engine/Source/Runtime/TraceLog/Public/Trace/Private/Writer.inl b/Engine/Source/Runtime/TraceLog/Public/Trace/Private/Writer.inl
deleted file mode 100644
diff --git a/Engine/Source/Runtime/TraceLog/Public/Trace/Trace.h b/Engine/Source/Runtime/TraceLog/Public/Trace/Trace.h
--- a/Engine/Source/Runtime/TraceLog/Public/Trace/Trace.h
+++ b/Engine/Source/Runtime/TraceLog/Public/Trace/Trace.h
@@ -2,8 +2,8 @@
 
 #pragma once
 
-#include "Private/Trace.h"
-#include "Private/Field.h"
+#include "Detail/Trace.h"
+#include "Detail/Field.h"
 
 ////////////////////////////////////////////////////////////////////////////////
 namespace Trace
@@ -24,5 +24,5 @@ UE_TRACE_API uint32 ToggleEvent(const TCHAR* Wildcard, bool bState) UE_TRACE_IMP
 #define UE_TRACE_EVENT_IS_ENABLED(LoggerName, EventName)		TRACE_PRIVATE_EVENT_IS_ENABLED(LoggerName, EventName)
 #define UE_TRACE_LOG(LoggerName, EventName, ...)				TRACE_PRIVATE_LOG(LoggerName, EventName, ##__VA_ARGS__)
 
-#include "Private/EventDef.inl"
-#include "Private/Writer.inl"
+#include "Detail/EventDef.inl"
+#include "Detail/Writer.inl"
