Merge 8685848 from //UE4/Dev-Rendering to Dev-Core (//UE4/Dev-Core)

---

[macOS][UnrealInsights] Implement UnrealInsights support for macOS.

+ Also add the unix platform define to the LogTrace header to enable all the juicy category goodness there.

<see original CL for reviewers, etc...>

---

#rnx

diff --git a/Engine/Source/Programs/UnrealInsights/Private/UserInterfaceCommand.cpp b/Engine/Source/Programs/UnrealInsights/Private/UserInterfaceCommand.cpp
--- a/Engine/Source/Programs/UnrealInsights/Private/UserInterfaceCommand.cpp
+++ b/Engine/Source/Programs/UnrealInsights/Private/UserInterfaceCommand.cpp
@@ -61,8 +61,8 @@ void FUserInterfaceCommand::Run()
 
 	// Manually load in the source code access plugins, as standalone programs don't currently support plugins.
 #if PLATFORM_MAC
-	IModuleInterface& XCodeSourceCodeAccessModule = FModuleManager::LoadModuleChecked<IModuleInterface>(FName("XCodeSourceCodeAccess"));
-	SourceCodeAccessModule.SetAccessor(FName("XCodeSourceCodeAccess"));
+	//im:TODO: IModuleInterface& XCodeSourceCodeAccessModule = FModuleManager::LoadModuleChecked<IModuleInterface>(FName("XCodeSourceCodeAccess"));
+	//im:TODO: SourceCodeAccessModule.SetAccessor(FName("XCodeSourceCodeAccess"));
 #elif PLATFORM_WINDOWS
 	//im:TODO: IModuleInterface& VisualStudioSourceCodeAccessModule = FModuleManager::LoadModuleChecked<IModuleInterface>(FName("VisualStudioSourceCodeAccess"));
 	//im:TODO: SourceCodeAccessModule.SetAccessor(FName("VisualStudioSourceCodeAccess"));
diff --git a/Engine/Source/Programs/UnrealInsights/UnrealInsights.Target.cs b/Engine/Source/Programs/UnrealInsights/UnrealInsights.Target.cs
--- a/Engine/Source/Programs/UnrealInsights/UnrealInsights.Target.cs
+++ b/Engine/Source/Programs/UnrealInsights/UnrealInsights.Target.cs
@@ -3,7 +3,7 @@
 using UnrealBuildTool;
 using System.Collections.Generic;
 
-[SupportedPlatforms("Win64", "Linux")]
+[SupportedPlatforms("Win64", "Linux", "Mac")]
 public class UnrealInsightsTarget : TargetRules
 {
 	public UnrealInsightsTarget(TargetInfo Target) : base(Target)
diff --git a/Engine/Source/Runtime/Core/Public/Logging/LogTrace.h b/Engine/Source/Runtime/Core/Public/Logging/LogTrace.h
--- a/Engine/Source/Runtime/Core/Public/Logging/LogTrace.h
+++ b/Engine/Source/Runtime/Core/Public/Logging/LogTrace.h
@@ -4,7 +4,7 @@
 
 #include "CoreTypes.h"
 
-#if !IS_PROGRAM && !UE_BUILD_SHIPPING && (PLATFORM_WINDOWS || PLATFORM_PS4 || PLATFORM_UNIX || PLATFORM_XBOXONE)
+#if !IS_PROGRAM && !UE_BUILD_SHIPPING && (PLATFORM_WINDOWS || PLATFORM_UNIX || PLATFORM_MAC || PLATFORM_PS4 || PLATFORM_XBOXONE)
 #define LOGTRACE_ENABLED 1
 #else
 #define LOGTRACE_ENABLED 0
diff --git a/Engine/Source/Runtime/Core/Public/ProfilingDebugging/CpuProfilerTrace.h b/Engine/Source/Runtime/Core/Public/ProfilingDebugging/CpuProfilerTrace.h
--- a/Engine/Source/Runtime/Core/Public/ProfilingDebugging/CpuProfilerTrace.h
+++ b/Engine/Source/Runtime/Core/Public/ProfilingDebugging/CpuProfilerTrace.h
@@ -4,7 +4,7 @@
 
 #include "CoreTypes.h"
 
-#if !IS_PROGRAM && !UE_BUILD_SHIPPING && (PLATFORM_WINDOWS || PLATFORM_PS4 || PLATFORM_XBOXONE || PLATFORM_SWITCH)
+#if !IS_PROGRAM && !UE_BUILD_SHIPPING && (PLATFORM_WINDOWS || PLATFORM_PS4 || PLATFORM_XBOXONE || PLATFORM_SWITCH || PLATFORM_UNIX || PLATFORM_MAC)
 #define CPUPROFILERTRACE_ENABLED 1
 #else
 #define CPUPROFILERTRACE_ENABLED 0
diff --git a/Engine/Source/Runtime/TraceLog/Private/Trace/Detail/Mac/MacTrace.cpp b/Engine/Source/Runtime/TraceLog/Private/Trace/Detail/Mac/MacTrace.cpp
new file mode 100644
--- /dev/null
+++ b/Engine/Source/Runtime/TraceLog/Private/Trace/Detail/Mac/MacTrace.cpp
@@ -0,0 +1,329 @@
+// Copyright 1998-2019 Epic Games, Inc. All Rights Reserved.
+
+#include "Trace/Trace.h"
+
+#if UE_TRACE_ENABLED
+
+#include <CoreServices/CoreServices.h>
+#include <fcntl.h>
+#include <mach/mach.h>
+#include <mach/mach_time.h>
+#include <netdb.h>
+#include <pthread.h>
+#include <sys/types.h>
+#include <sys/mman.h>
+#include <sys/select.h>
+#include <sys/socket.h>
+#include <unistd.h>
+
+namespace Trace {
+namespace Private {
+
+////////////////////////////////////////////////////////////////////////////////
+enum EHandleType : UPTRINT
+{
+	HandleType_File		= 0,
+	HandleType_Socket	= 1,
+	HandleType_Reserved	= 2,
+	_HandleType_Mask	= (1 << 2) - 1,
+};
+
+////////////////////////////////////////////////////////////////////////////////
+template <typename HandleType>
+static UPTRINT EncodeHandle(HandleType Handle, EHandleType Type)
+{
+	if (Type == HandleType_File)
+	{
+		return UPTRINT(Handle);
+	}
+
+	return (UPTRINT(Handle) << 2) | Type;
+}
+
+////////////////////////////////////////////////////////////////////////////////
+static EHandleType DecodeHandle(UPTRINT& Handle)
+{
+	if (int Type = int(Handle & _HandleType_Mask))
+	{
+		Handle >>= 2;
+		return static_cast<EHandleType>(Type);
+	}
+	return HandleType_File;
+}
+
+////////////////////////////////////////////////////////////////////////////////
+uint8* MemoryReserve(SIZE_T Size)
+{
+	void* Ptr = mmap(nullptr, Size, PROT_READ | PROT_WRITE, MAP_PRIVATE | MAP_ANON, -1, 0);
+	return (Ptr != MAP_FAILED) ? reinterpret_cast<uint8*>(Ptr) : nullptr;
+}
+
+////////////////////////////////////////////////////////////////////////////////
+void MemoryFree(void* Address, SIZE_T Size)
+{
+	munmap(Address, Size);
+}
+
+////////////////////////////////////////////////////////////////////////////////
+void MemoryMap(void* Address, SIZE_T Size)
+{
+	// no-op if mmap()ed with R/W
+}
+
+////////////////////////////////////////////////////////////////////////////////
+void MemoryUnmap(void* Address, SIZE_T Size)
+{
+	madvise(Address, Size, MADV_DONTNEED);
+}
+
+
+
+////////////////////////////////////////////////////////////////////////////////
+UPTRINT ThreadCreate(const ANSICHAR* Name, void (*Entry)())
+{
+	void* (*PthreadThunk)(void*) = [] (void* Param) -> void * {
+		typedef void (*EntryType)(void);
+		(EntryType(Param))();
+		return nullptr;
+	};
+
+	pthread_t ThreadHandle;
+	if (pthread_create(&ThreadHandle, nullptr, PthreadThunk, reinterpret_cast<void *>(Entry)) != 0)
+	{
+		return 0;
+	}
+	return reinterpret_cast<UPTRINT>(ThreadHandle);
+}
+
+////////////////////////////////////////////////////////////////////////////////
+uint32 ThreadGetCurrentId()
+{
+	uint64_t ThreadId = 0;
+	pthread_threadid_np(nullptr, &ThreadId);
+	return static_cast<uint32>(ThreadId);
+}
+
+////////////////////////////////////////////////////////////////////////////////
+void ThreadSleep(uint32 Milliseconds)
+{
+	usleep(Milliseconds * 1000U);
+}
+
+////////////////////////////////////////////////////////////////////////////////
+void ThreadJoin(UPTRINT Handle)
+{
+	pthread_join(reinterpret_cast<pthread_t>(Handle), nullptr);
+}
+
+////////////////////////////////////////////////////////////////////////////////
+void ThreadDestroy(UPTRINT Handle)
+{
+	// no-op
+}
+
+
+
+////////////////////////////////////////////////////////////////////////////////
+const mach_timebase_info_data_t& TimeGetInfo()
+{
+    static mach_timebase_info_data_t Info;
+    static dispatch_once_t onceToken;
+
+	dispatch_once(&onceToken, ^{
+		   (void) mach_timebase_info(&Info);
+	   });
+
+	return Info;
+}
+
+////////////////////////////////////////////////////////////////////////////////
+uint64 TimeGetFrequency()
+{
+	const mach_timebase_info_data_t& Info = TimeGetInfo();
+	return (uint64)((double)Info.numer / (1e-9 * (double)Info.denom) + 0.5);
+}
+
+////////////////////////////////////////////////////////////////////////////////
+uint64 TimeGetTimestamp()
+{
+	return mach_absolute_time();
+}
+
+
+
+////////////////////////////////////////////////////////////////////////////////
+UPTRINT TcpSocketConnect(const ANSICHAR* Host, uint16 Port)
+{
+	struct FAddrInfoPtr
+	{
+					~FAddrInfoPtr()	{ freeaddrinfo(Value); }
+		addrinfo*	operator -> ()	{ return Value; }
+		addrinfo**	operator & ()	{ return &Value; }
+		addrinfo*	Value;
+	};
+
+	FAddrInfoPtr Info;
+	addrinfo Hints = {};
+	Hints.ai_family = AF_INET;
+	Hints.ai_socktype = SOCK_STREAM;
+	Hints.ai_protocol = IPPROTO_TCP;
+	if (getaddrinfo(Host, nullptr, &Hints, &Info))
+	{
+		return 0;
+	}
+
+	if (&Info == nullptr)
+	{
+		return 0;
+	}
+
+	auto* SockAddr = (sockaddr_in*)Info->ai_addr;
+	SockAddr->sin_port = htons(Port);
+
+	int Socket = socket(AF_INET, SOCK_STREAM, IPPROTO_TCP);
+	if (Socket < 0)
+	{
+		return 0;
+	}
+
+	int Result = connect(Socket, Info->ai_addr, int(Info->ai_addrlen));
+	if (Result < 0)
+	{
+		close(Socket);
+		return 0;
+	}
+
+	return EncodeHandle(UPTRINT(Socket), HandleType_Socket);
+}
+
+////////////////////////////////////////////////////////////////////////////////
+UPTRINT TcpSocketListen(uint16 Port)
+{
+	int Socket = socket(AF_INET, SOCK_STREAM, IPPROTO_TCP);
+	if (Socket < 0)
+	{
+		return 0;
+	}
+
+	sockaddr_in SockAddr;
+	SockAddr.sin_family = AF_INET;
+	SockAddr.sin_addr.s_addr = 0;
+	SockAddr.sin_port = htons(Port);
+	int Result = bind(Socket, reinterpret_cast<sockaddr*>(&SockAddr), sizeof(SockAddr));
+	if (Result < 0)
+	{
+		close(Socket);
+		return 0;
+	}
+
+	Result = listen(Socket, 1);
+	if (Result < 0)
+	{
+		close(Socket);
+		return 0;
+	}
+
+	return EncodeHandle(UPTRINT(Socket), HandleType_Socket);
+}
+
+////////////////////////////////////////////////////////////////////////////////
+int32 TcpSocketAccept(UPTRINT Socket, UPTRINT& Out)
+{
+	int HandleType = DecodeHandle(Socket);
+	int Inner = int(Socket);
+
+	Inner = accept(Inner, nullptr, nullptr);
+	if (Inner < 0)
+	{
+		return 0;
+	}
+
+	Out = EncodeHandle(UPTRINT(Inner), HandleType_Socket);
+	return 1;
+}
+
+////////////////////////////////////////////////////////////////////////////////
+bool TcpSocketHasData(UPTRINT Socket)
+{
+	int HandleType = DecodeHandle(Socket);
+	int Inner = int(Socket);
+	fd_set FdSet;
+	FD_ZERO(&FdSet);
+	FD_SET(Inner, &FdSet);
+	timeval TimeVal = {};
+	int result = select(Inner + 1, &FdSet, nullptr, nullptr, &TimeVal);
+	return ((result != 0) || ((result == -1) && (errno == ETIMEDOUT)));
+}
+
+////////////////////////////////////////////////////////////////////////////////
+bool IoWrite(UPTRINT Handle, const void* Data, uint32 Size)
+{
+	int HandleType = DecodeHandle(Handle);
+	int Inner = int(Handle);
+
+	if (HandleType == HandleType_File)
+	{
+		int BytesWritten = write(Inner, Data, Size);
+		if (BytesWritten == -1)
+		{
+			return false;
+		}
+		return (BytesWritten == Size);
+	}
+	else
+	{
+		// If not a file, then it's a socket for macOS.
+		return (send(Inner, (const char*)Data, Size, 0) == Size);
+	}
+}
+
+////////////////////////////////////////////////////////////////////////////////
+int32 IoRead(UPTRINT Handle, void* Data, uint32 Size)
+{
+	int HandleType = DecodeHandle(Handle);
+	int Inner = int(Handle);
+
+	if (HandleType == HandleType_File)
+	{
+		return read(Inner, Data, Size);
+	}
+	else
+	{
+		// If not a file, then it's a socket for macOS.
+		return recv(Inner, (char*)Data, Size, 0);
+	}
+}
+
+////////////////////////////////////////////////////////////////////////////////
+void IoClose(UPTRINT Handle)
+{
+	int HandleType = DecodeHandle(Handle);
+	int Inner = int(Handle);
+
+	close(Inner);
+}
+
+////////////////////////////////////////////////////////////////////////////////
+UPTRINT FileOpen(const ANSICHAR* Path, const ANSICHAR* Mode)
+{
+	int Flags = O_CREAT | O_APPEND | O_RDWR | O_SHLOCK | O_NONBLOCK;
+
+	if ((*Mode) == 'w')
+	{
+		Flags |= O_TRUNC;
+	}
+
+	int Out = open(Path, Flags, S_IRUSR | S_IWUSR | S_IRGRP | S_IWGRP | S_IROTH | S_IWOTH);
+
+	if (((*Mode) == 'a') && (Out != -1))
+	{
+		lseek(Out, 0, SEEK_END);
+	}
+
+	return EncodeHandle(UPTRINT(Out), HandleType_File);
+}
+
+} // namespace Private
+} // namespace Trace
+
+#endif // UE_TRACE_ENABLED
diff --git a/Engine/Source/Runtime/TraceLog/Public/Trace/Detail/Trace.h b/Engine/Source/Runtime/TraceLog/Public/Trace/Detail/Trace.h
--- a/Engine/Source/Runtime/TraceLog/Public/Trace/Detail/Trace.h
+++ b/Engine/Source/Runtime/TraceLog/Public/Trace/Detail/Trace.h
@@ -7,7 +7,7 @@
 ////////////////////////////////////////////////////////////////////////////////
 #if !defined(UE_TRACE_ENABLED)
 #	if !UE_BUILD_SHIPPING && !IS_PROGRAM
-#		if PLATFORM_WINDOWS || PLATFORM_PS4 || PLATFORM_XBOXONE || PLATFORM_SWITCH
+#		if PLATFORM_WINDOWS || PLATFORM_UNIX || PLATFORM_MAC || PLATFORM_PS4 || PLATFORM_XBOXONE || PLATFORM_SWITCH
 #			define UE_TRACE_ENABLED	1
 #		endif
 #	endif
