Removed Mac's ThreadGetCurrentId() as it's not used anymore

diff --git a/Engine/Source/Runtime/TraceLog/Private/Trace/Detail/Mac/MacTrace.cpp b/Engine/Source/Runtime/TraceLog/Private/Trace/Detail/Mac/MacTrace.cpp
--- a/Engine/Source/Runtime/TraceLog/Private/Trace/Detail/Mac/MacTrace.cpp
+++ b/Engine/Source/Runtime/TraceLog/Private/Trace/Detail/Mac/MacTrace.cpp
@@ -95,14 +95,6 @@ UPTRINT ThreadCreate(const ANSICHAR* Name, void (*Entry)())
 	return reinterpret_cast<UPTRINT>(ThreadHandle);
 }
 
-////////////////////////////////////////////////////////////////////////////////
-uint32 ThreadGetCurrentId()
-{
-	uint64_t ThreadId = 0;
-	pthread_threadid_np(nullptr, &ThreadId);
-	return static_cast<uint32>(ThreadId);
-}
-
 ////////////////////////////////////////////////////////////////////////////////
 void ThreadSleep(uint32 Milliseconds)
 {
