Pascal-casing fix.

diff --git a/Engine/Source/Runtime/TraceLog/Private/Trace/Detail/Mac/MacTrace.cpp b/Engine/Source/Runtime/TraceLog/Private/Trace/Detail/Mac/MacTrace.cpp
--- a/Engine/Source/Runtime/TraceLog/Private/Trace/Detail/Mac/MacTrace.cpp
+++ b/Engine/Source/Runtime/TraceLog/Private/Trace/Detail/Mac/MacTrace.cpp
@@ -119,9 +119,9 @@ void ThreadDestroy(UPTRINT Handle)
 const mach_timebase_info_data_t& TimeGetInfo()
 {
     static mach_timebase_info_data_t Info;
-    static dispatch_once_t onceToken;
+    static dispatch_once_t OnceToken;
 
-	dispatch_once(&onceToken, ^{
+	dispatch_once(&OnceToken, ^{
 		   (void) mach_timebase_info(&Info);
 	   });
 
