Listening sockets should be non-blocking. This is for compatibility with another platform's development-only sockets emulation API which doesn't support selecting on listening sockets

(ushell-p4-stage of 8971110)

diff --git a/Engine/Source/Runtime/TraceLog/Private/Trace/Detail/Mac/MacTrace.cpp b/Engine/Source/Runtime/TraceLog/Private/Trace/Detail/Mac/MacTrace.cpp
--- a/Engine/Source/Runtime/TraceLog/Private/Trace/Detail/Mac/MacTrace.cpp
+++ b/Engine/Source/Runtime/TraceLog/Private/Trace/Detail/Mac/MacTrace.cpp
@@ -143,6 +143,19 @@ uint64 TimeGetTimestamp()
 
 
 
+////////////////////////////////////////////////////////////////////////////////
+static bool TcpSocketSetNonBlocking(int Socket, bool bNonBlocking)
+{
+	int Flags = fcntl(Socket, F_GETFL, 0);
+	if (Flags == -1)
+	{
+		return false;
+	}
+
+	Flags = bNonBlocking ? (Flags|O_NONBLOCK) : (Flags & ~O_NONBLOCK);
+	return fcntl(Socket, F_SETFL, Flags) >= 0;
+}
+
 ////////////////////////////////////////////////////////////////////////////////
 UPTRINT TcpSocketConnect(const ANSICHAR* Host, uint16 Port)
 {
@@ -185,6 +198,12 @@ UPTRINT TcpSocketConnect(const ANSICHAR* Host, uint16 Port)
 		return 0;
 	}
 
+	if (!TcpSocketSetNonBlocking(Socket, false))
+	{
+		close(Socket);
+		return 0;
+	}
+
 	return EncodeHandle(UPTRINT(Socket), HandleType_Socket);
 }
 
@@ -215,6 +234,12 @@ UPTRINT TcpSocketListen(uint16 Port)
 		return 0;
 	}
 
+	if (!TcpSocketSetNonBlocking(Socket, true))
+	{
+		close(Socket);
+		return 0;
+	}
+
 	return EncodeHandle(UPTRINT(Socket), HandleType_Socket);
 }
 
@@ -227,6 +252,12 @@ int32 TcpSocketAccept(UPTRINT Socket, UPTRINT& Out)
 	Inner = accept(Inner, nullptr, nullptr);
 	if (Inner < 0)
 	{
+		return (Inner == EAGAIN || Inner == EWOULDBLOCK) - 1; // 0 if would block else -1
+	}
+
+	if (!TcpSocketSetNonBlocking(Inner, false))
+	{
+		close(Inner);
 		return 0;
 	}
 
