Removed the 'Mode' argument from FileOpen(). File handling is now strictly write-only as that's all we need.

diff --git a/Engine/Source/Runtime/TraceLog/Private/Trace/Detail/Mac/MacTrace.cpp b/Engine/Source/Runtime/TraceLog/Private/Trace/Detail/Mac/MacTrace.cpp
--- a/Engine/Source/Runtime/TraceLog/Private/Trace/Detail/Mac/MacTrace.cpp
+++ b/Engine/Source/Runtime/TraceLog/Private/Trace/Detail/Mac/MacTrace.cpp
@@ -327,20 +327,14 @@ void IoClose(UPTRINT Handle)
 }
 
 ////////////////////////////////////////////////////////////////////////////////
-UPTRINT FileOpen(const ANSICHAR* Path, const ANSICHAR* Mode)
+UPTRINT FileOpen(const ANSICHAR* Path)
 {
-	int Flags = O_CREAT | O_APPEND | O_RDWR | O_SHLOCK | O_NONBLOCK;
-
-	if ((*Mode) == 'w')
-	{
-		Flags |= O_TRUNC;
-	}
-
-	int Out = open(Path, Flags, S_IRUSR | S_IWUSR | S_IRGRP | S_IWGRP | S_IROTH | S_IWOTH);
-
-	if (((*Mode) == 'a') && (Out != -1))
+	int Flags = O_CREAT|O_WRONLY|O_TRUNC|O_SHLOCK;
+	int Mode = S_IRUSR|S_IWUSR|S_IRGRP|S_IWGRP|S_IROTH|S_IWOTH;
+	int Out = open(Path, Flags, Mode);
+	if (Out < 0)
 	{
-		lseek(Out, 0, SEEK_END);
+		return 0;
 	}
 
 	return EncodeHandle(UPTRINT(Out), HandleType_File);
diff --git a/Engine/Source/Runtime/TraceLog/Private/Trace/Detail/Windows/WindowsTrace.cpp b/Engine/Source/Runtime/TraceLog/Private/Trace/Detail/Windows/WindowsTrace.cpp
--- a/Engine/Source/Runtime/TraceLog/Private/Trace/Detail/Windows/WindowsTrace.cpp
+++ b/Engine/Source/Runtime/TraceLog/Private/Trace/Detail/Windows/WindowsTrace.cpp
@@ -269,31 +269,18 @@ void IoClose(UPTRINT Handle)
 
 
 ////////////////////////////////////////////////////////////////////////////////
-UPTRINT FileOpen(const ANSICHAR* Path, const ANSICHAR* Mode)
+UPTRINT FileOpen(const ANSICHAR* Path)
 {
-	DWORD Access = GENERIC_READ;
-	DWORD Disposition = OPEN_EXISTING;
-	switch (*Mode)
-	{
-	case 'a':
-	case 'w':
-		Access |= GENERIC_WRITE;
-		Disposition = (*Mode == 'a') ? OPEN_ALWAYS : CREATE_ALWAYS;
-		break;
-	}
-
+	DWORD Access = GENERIC_WRITE;
+	DWORD Share = FILE_SHARE_READ;
+	DWORD Disposition = CREATE_ALWAYS;
 	DWORD Flags = FILE_ATTRIBUTE_NORMAL;
-	HANDLE Out = CreateFileA(Path, Access, FILE_SHARE_READ|FILE_SHARE_WRITE, nullptr, Disposition, Flags, nullptr);
+	HANDLE Out = CreateFileA(Path, Access, Share, nullptr, Disposition, Flags, nullptr);
 	if (Out == INVALID_HANDLE_VALUE)
 	{
 		return 0;
 	}
 
-	if (*Mode == 'a')
-	{
-		SetFilePointer(Out, 0, nullptr, FILE_END);
-	}
-
 	return UPTRINT(Out) + 1;
 }
 
diff --git a/Engine/Source/Runtime/TraceLog/Private/Trace/Platform.h b/Engine/Source/Runtime/TraceLog/Private/Trace/Platform.h
--- a/Engine/Source/Runtime/TraceLog/Private/Trace/Platform.h
+++ b/Engine/Source/Runtime/TraceLog/Private/Trace/Platform.h
@@ -37,7 +37,7 @@ bool	IoWrite(UPTRINT Handle, const void* Data, uint32 Size);
 void	IoClose(UPTRINT Handle);
 
 ////////////////////////////////////////////////////////////////////////////////
-UPTRINT	FileOpen(const ANSICHAR* Path, const ANSICHAR* Mode);
+UPTRINT	FileOpen(const ANSICHAR* Path);
 
 } // namespace Private
 } // namespace Trace
diff --git a/Engine/Source/Runtime/TraceLog/Private/Trace/Writer.cpp b/Engine/Source/Runtime/TraceLog/Private/Trace/Writer.cpp
--- a/Engine/Source/Runtime/TraceLog/Private/Trace/Writer.cpp
+++ b/Engine/Source/Runtime/TraceLog/Private/Trace/Writer.cpp
@@ -6,6 +6,7 @@
 
 #include "Trace/Platform.h"
 #include "Trace/Detail/Atomic.h"
+#include "Trace/Detail/Protocol.h"
 #include "Trace/Trace.h"
 
 #include "Misc/CString.h"
@@ -913,7 +914,7 @@ bool Writer_Open(const ANSICHAR* Path)
 {
 	Writer_Initialize();
 
-	UPTRINT DataHandle = FileOpen(Path, "w");
+	UPTRINT DataHandle = FileOpen(Path);
 	if (!DataHandle)
 	{
 		return false;
