Removed include statement that escaped by accident

diff --git a/Engine/Source/Runtime/TraceLog/Private/Trace/Writer.cpp b/Engine/Source/Runtime/TraceLog/Private/Trace/Writer.cpp
--- a/Engine/Source/Runtime/TraceLog/Private/Trace/Writer.cpp
+++ b/Engine/Source/Runtime/TraceLog/Private/Trace/Writer.cpp
@@ -6,7 +6,6 @@
 
 #include "Trace/Platform.h"
 #include "Trace/Detail/Atomic.h"
-#include "Trace/Detail/Protocol.h"
 #include "Trace/Trace.h"
 
 #include "Misc/CString.h"
