As per the POSIX API, sockets are file descriptors. This means they can use the same read/write() API as files so there is no need to encode type into the returned handle. Descriptors are biased by one as 0 is a valid file descriptor whereas the Trace API treats 0 as invalid.

(ushell-p4-stage of 8971110)

diff --git a/Engine/Source/Runtime/TraceLog/Private/Trace/Detail/Mac/MacTrace.cpp b/Engine/Source/Runtime/TraceLog/Private/Trace/Detail/Mac/MacTrace.cpp
--- a/Engine/Source/Runtime/TraceLog/Private/Trace/Detail/Mac/MacTrace.cpp
+++ b/Engine/Source/Runtime/TraceLog/Private/Trace/Detail/Mac/MacTrace.cpp
@@ -19,38 +19,6 @@
 namespace Trace {
 namespace Private {
 
-////////////////////////////////////////////////////////////////////////////////
-enum EHandleType : UPTRINT
-{
-	HandleType_File		= 0,
-	HandleType_Socket	= 1,
-	HandleType_Reserved	= 2,
-	_HandleType_Mask	= (1 << 2) - 1,
-};
-
-////////////////////////////////////////////////////////////////////////////////
-template <typename HandleType>
-static UPTRINT EncodeHandle(HandleType Handle, EHandleType Type)
-{
-	if (Type == HandleType_File)
-	{
-		return UPTRINT(Handle);
-	}
-
-	return (UPTRINT(Handle) << 2) | Type;
-}
-
-////////////////////////////////////////////////////////////////////////////////
-static EHandleType DecodeHandle(UPTRINT& Handle)
-{
-	if (int Type = int(Handle & _HandleType_Mask))
-	{
-		Handle >>= 2;
-		return static_cast<EHandleType>(Type);
-	}
-	return HandleType_File;
-}
-
 ////////////////////////////////////////////////////////////////////////////////
 uint8* MemoryReserve(SIZE_T Size)
 {
@@ -240,7 +208,7 @@ UPTRINT TcpSocketListen(uint16 Port)
 		return 0;
 	}
 
-	return EncodeHandle(UPTRINT(Socket), HandleType_Socket);
+	return UPTRINT(Socket + 1);
 }
 
 ////////////////////////////////////////////////////////////////////////////////
@@ -261,15 +229,14 @@ int32 TcpSocketAccept(UPTRINT Socket, UPTRINT& Out)
 		return 0;
 	}
 
-	Out = EncodeHandle(UPTRINT(Inner), HandleType_Socket);
+	Out = UPTRINT(Inner + 1);
 	return 1;
 }
 
 ////////////////////////////////////////////////////////////////////////////////
 bool TcpSocketHasData(UPTRINT Socket)
 {
-	int HandleType = DecodeHandle(Socket);
-	int Inner = int(Socket);
+	int Inner = int(Socket - 1);
 	fd_set FdSet;
 	FD_ZERO(&FdSet);
 	FD_SET(Inner, &FdSet);
@@ -278,51 +245,26 @@ bool TcpSocketHasData(UPTRINT Socket)
 	return ((result != 0) || ((result == -1) && (errno == ETIMEDOUT)));
 }
 
+
+
 ////////////////////////////////////////////////////////////////////////////////
 bool IoWrite(UPTRINT Handle, const void* Data, uint32 Size)
 {
-	int HandleType = DecodeHandle(Handle);
-	int Inner = int(Handle);
-
-	if (HandleType == HandleType_File)
-	{
-		int BytesWritten = write(Inner, Data, Size);
-		if (BytesWritten == -1)
-		{
-			return false;
-		}
-		return (BytesWritten == Size);
-	}
-	else
-	{
-		// If not a file, then it's a socket for macOS.
-		return (send(Inner, (const char*)Data, Size, 0) == Size);
-	}
+	int Inner = int(Handle - 1);
+	return (write(Inner, Data, Size) == Size);
 }
 
 ////////////////////////////////////////////////////////////////////////////////
 int32 IoRead(UPTRINT Handle, void* Data, uint32 Size)
 {
-	int HandleType = DecodeHandle(Handle);
-	int Inner = int(Handle);
-
-	if (HandleType == HandleType_File)
-	{
-		return read(Inner, Data, Size);
-	}
-	else
-	{
-		// If not a file, then it's a socket for macOS.
-		return recv(Inner, (char*)Data, Size, 0);
-	}
+	int Inner = int(Handle - 1);
+	return read(Inner, Data, Size);
 }
 
 ////////////////////////////////////////////////////////////////////////////////
 void IoClose(UPTRINT Handle)
 {
-	int HandleType = DecodeHandle(Handle);
-	int Inner = int(Handle);
-
+	int Inner = int(Handle - 1);
 	close(Inner);
 }
 
@@ -337,7 +279,7 @@ UPTRINT FileOpen(const ANSICHAR* Path)
 		return 0;
 	}
 
-	return EncodeHandle(UPTRINT(Out), HandleType_File);
+	return UPTRINT(Out + 1);
 }
 
 } // namespace Private
