Follow-up to 8973775. Not sure how I managed to miss these two Encode/Decodes calls.

(ushell-p4-cherrypick of 8971109 by [REDACTED_PERSON])

diff --git a/Engine/Source/Runtime/TraceLog/Private/Trace/Detail/Mac/MacTrace.cpp b/Engine/Source/Runtime/TraceLog/Private/Trace/Detail/Mac/MacTrace.cpp
--- a/Engine/Source/Runtime/TraceLog/Private/Trace/Detail/Mac/MacTrace.cpp
+++ b/Engine/Source/Runtime/TraceLog/Private/Trace/Detail/Mac/MacTrace.cpp
@@ -172,7 +172,7 @@ UPTRINT TcpSocketConnect(const ANSICHAR* Host, uint16 Port)
 		return 0;
 	}
 
-	return EncodeHandle(UPTRINT(Socket), HandleType_Socket);
+	return UPTRINT(Socket + 1);
 }
 
 ////////////////////////////////////////////////////////////////////////////////
@@ -214,8 +214,7 @@ UPTRINT TcpSocketListen(uint16 Port)
 ////////////////////////////////////////////////////////////////////////////////
 int32 TcpSocketAccept(UPTRINT Socket, UPTRINT& Out)
 {
-	int HandleType = DecodeHandle(Socket);
-	int Inner = int(Socket);
+	int Inner = int(Socket - 1);
 
 	Inner = accept(Inner, nullptr, nullptr);
 	if (Inner < 0)
