Consolidated code setting/clearing non-blocking mode on sockets into a function for clarity.

diff --git a/Engine/Source/Runtime/TraceLog/Private/Trace/Detail/Windows/WindowsTrace.cpp b/Engine/Source/Runtime/TraceLog/Private/Trace/Detail/Windows/WindowsTrace.cpp
--- a/Engine/Source/Runtime/TraceLog/Private/Trace/Detail/Windows/WindowsTrace.cpp
+++ b/Engine/Source/Runtime/TraceLog/Private/Trace/Detail/Windows/WindowsTrace.cpp
@@ -108,6 +108,13 @@ static void TcpSocketInitialize()
 	WSAStartup(MAKEWORD(2, 2), &WsaData);
 }
 
+////////////////////////////////////////////////////////////////////////////////
+static bool TcpSocketSetNonBlocking(SOCKET Socket, bool bNonBlocking)
+{
+	unsigned long NonBlockingMode = !!bNonBlocking;
+	return ioctlsocket(Socket, FIONBIO, &NonBlockingMode) != SOCKET_ERROR;
+}
+
 ////////////////////////////////////////////////////////////////////////////////
 UPTRINT TcpSocketConnect(const ANSICHAR* Host, uint16 Port)
 {
@@ -155,6 +162,12 @@ UPTRINT TcpSocketConnect(const ANSICHAR* Host, uint16 Port)
 		return 0;
 	}
 
+	if (!TcpSocketSetNonBlocking(Socket, 0))
+	{
+		closesocket(Socket);
+		return 0;
+	}
+
 	return UPTRINT(Socket) + 1;
 }
 
@@ -188,8 +201,7 @@ UPTRINT TcpSocketListen(uint16 Port)
 		return 0;
 	}
 
-	unsigned long NonBlockingMode = 1;
-	if (ioctlsocket(Socket, FIONBIO, &NonBlockingMode) == SOCKET_ERROR)
+	if (!TcpSocketSetNonBlocking(Socket, 1))
 	{
 		closesocket(Socket);
 		return 0;
@@ -206,13 +218,12 @@ int32 TcpSocketAccept(UPTRINT Socket, UPTRINT& Out)
 	Inner = accept(Inner, nullptr, nullptr);
 	if (Inner == INVALID_SOCKET)
 	{
-		return (WSAGetLastError() == WSAEWOULDBLOCK) - 1;
+		return (WSAGetLastError() == WSAEWOULDBLOCK) - 1; // 0 if would block else -1
 	}
 
-	unsigned long NonBlockingMode = 0;
-	if (ioctlsocket(Socket, FIONBIO, &NonBlockingMode) == SOCKET_ERROR)
+	if (!TcpSocketSetNonBlocking(Inner, 0))
 	{
-		closesocket(Socket);
+		closesocket(Inner);
 		return 0;
 	}
 
