Fixed up TraceLog's support for Unix following API changes.

diff --git a/Engine/Source/Runtime/TraceLog/Private/Trace/Unix/UnixTrace.cpp b/Engine/Source/Runtime/TraceLog/Private/Trace/Unix/UnixTrace.cpp
--- a/Engine/Source/Runtime/TraceLog/Private/Trace/Unix/UnixTrace.cpp
+++ b/Engine/Source/Runtime/TraceLog/Private/Trace/Unix/UnixTrace.cpp
@@ -4,13 +4,15 @@
 
 #if UE_TRACE_ENABLED
 
+#include <errno.h>
+#include <fcntl.h>
+#include <netdb.h>
+#include <pthread.h>
 #include <sys/mman.h>
-#include <sys/types.h>
 #include <sys/socket.h>
-#include <netdb.h>
-#include <unistd.h>
+#include <sys/types.h>
 #include <time.h>
-#include <pthread.h>
+#include <unistd.h>
 #if defined(_GNU_SOURCE)
 	#include <sys/syscall.h>
 #endif // _GNU_SOURCE
@@ -99,6 +101,19 @@ uint64 TimeGetTimestamp()
 
 
 
+////////////////////////////////////////////////////////////////////////////////
+static bool TcpSocketSetNonBlocking(int Socket, bool bNonBlocking)
+{
+	int Flags = fcntl(Socket, F_GETFL, 0);
+	if (Flags == -1)
+	{
+		return false;
+	}
+
+	Flags = bNonBlocking ? (Flags|O_NONBLOCK) : (Flags & ~O_NONBLOCK);
+	return fcntl(Socket, F_SETFL, Flags) >= 0;
+}
+
 ////////////////////////////////////////////////////////////////////////////////
 UPTRINT TcpSocketConnect(const ANSICHAR* Host, uint16 Port)
 {
@@ -141,6 +156,12 @@ UPTRINT TcpSocketConnect(const ANSICHAR* Host, uint16 Port)
 		return 0;
 	}
 
+	if (!TcpSocketSetNonBlocking(Socket, false))
+	{
+		close(Socket);
+		return 0;
+	}
+
 	return UPTRINT(Socket + 1);
 }
 
@@ -171,32 +192,38 @@ UPTRINT TcpSocketListen(uint16 Port)
 		return 0;
 	}
 
+	if (!TcpSocketSetNonBlocking(Socket, true))
+	{
+		close(Socket);
+		return 0;
+	}
+
 	return UPTRINT(Socket + 1);
 }
 
 ////////////////////////////////////////////////////////////////////////////////
-UPTRINT TcpSocketAccept(UPTRINT Socket)
+int32 TcpSocketAccept(UPTRINT Socket, UPTRINT& Out)
 {
 	int Inner = Socket - 1;
 
 	Inner = accept(Inner, nullptr, nullptr);
 	if (Inner < 0)
 	{
-		return 0;
+		return (Inner == EAGAIN || Inner == EWOULDBLOCK) - 1; // 0 if would block else -1
 	}
 
-	return UPTRINT(Inner + 1);
-}
+	if (!TcpSocketSetNonBlocking(Inner, false))
+	{
+		close(Inner);
+		return 0;
+	}
 
-////////////////////////////////////////////////////////////////////////////////
-void TcpSocketClose(UPTRINT Socket)
-{
-	int Inner = Socket - 1;
-	close(Inner);
+	Out = UPTRINT(Inner + 1);
+	return 1;
 }
 
 ////////////////////////////////////////////////////////////////////////////////
-bool TcpSocketSelect(UPTRINT Socket)
+bool TcpSocketHasData(UPTRINT Socket)
 {
 	int Inner = Socket - 1;
 	fd_set FdSet;
@@ -206,18 +233,44 @@ bool TcpSocketSelect(UPTRINT Socket)
 	return (select(Inner + 1, &FdSet, nullptr, nullptr, &TimeVal) != 0);
 }
 
+
+
 ////////////////////////////////////////////////////////////////////////////////
-bool TcpSocketSend(UPTRINT Socket, const void* Data, uint32 Size)
+bool IoWrite(UPTRINT Handle, const void* Data, uint32 Size)
 {
-	int Inner = Socket - 1;
-	return send(Inner, (const char*)Data, Size, 0) == Size;
+	int Inner = int(Handle) - 1;
+	return write(Inner, Data, Size) == Size;
 }
 
 ////////////////////////////////////////////////////////////////////////////////
-int32 TcpSocketRecv(UPTRINT Socket, void* Data, uint32 Size)
+int32 IoRead(UPTRINT Handle, void* Data, uint32 Size)
 {
-	int Inner = Socket - 1;
-	return static_cast<int32>(recv(Inner, (char*)Data, Size, 0));
+	int Inner = int(Handle) - 1;
+	return read(Inner, Data, Size);
+}
+
+////////////////////////////////////////////////////////////////////////////////
+void IoClose(UPTRINT Handle)
+{
+	int Inner = int(Handle) - 1;
+	close(Inner);
+}
+
+
+
+////////////////////////////////////////////////////////////////////////////////
+UPTRINT FileOpen(const ANSICHAR* Path)
+{
+	int Flags = O_CREAT|O_WRONLY|O_TRUNC;
+	int Mode = S_IRUSR|S_IWUSR|S_IRGRP|S_IWGRP|S_IROTH|S_IWOTH;
+
+	int Out = open(Path, Flags, Mode);
+	if (Out < 0)
+	{
+		return 0;
+	}
+
+	return UPTRINT(Out + 1);
 }
 
 } // namespace Private
