Do not send empty payloads as they don't contain any collected events

diff --git a/Engine/Source/Runtime/TraceLog/Private/Trace/Writer.cpp b/Engine/Source/Runtime/TraceLog/Private/Trace/Writer.cpp
--- a/Engine/Source/Runtime/TraceLog/Private/Trace/Writer.cpp
+++ b/Engine/Source/Runtime/TraceLog/Private/Trace/Writer.cpp
@@ -382,6 +382,11 @@ static void Writer_ConsumeEvents()
 
 		void Flush()
 		{
+			if (Cursor == sizeof(FPayload::Data))
+			{
+				return;
+			}
+
 			// There will always be space remaining for the header because we've
 			// arranged for that by including the header in FPayload. We'll shift
 			// it forward so it butts up against the event data (at the expense of
