Connect/Open is a strange API. We're not connecting to or opening a trace, we're telling Trace to send its events somewhere. I've renamed these methods.

diff --git a/Engine/Source/Developer/TraceAnalysis/Private/ControlClient.cpp b/Engine/Source/Developer/TraceAnalysis/Private/ControlClient.cpp
--- a/Engine/Source/Developer/TraceAnalysis/Private/ControlClient.cpp
+++ b/Engine/Source/Developer/TraceAnalysis/Private/ControlClient.cpp
@@ -100,14 +100,14 @@ bool FControlClient::IsConnected() const
 }
 
 ////////////////////////////////////////////////////////////////////////////////
-void FControlClient::SendConnect(const TCHAR* Path)
+void FControlClient::SendSendTo(const TCHAR* Host)
 {
     if (!IsConnected())
     {
         return;
     }
 
-    FormatAndSend(TEXT("Connect %s"), Path);
+    FormatAndSend(TEXT("SendTo %s"), Host);
 }
 
 ////////////////////////////////////////////////////////////////////////////////
diff --git a/Engine/Source/Developer/TraceAnalysis/Public/Trace/ControlClient.h b/Engine/Source/Developer/TraceAnalysis/Public/Trace/ControlClient.h
--- a/Engine/Source/Developer/TraceAnalysis/Public/Trace/ControlClient.h
+++ b/Engine/Source/Developer/TraceAnalysis/Public/Trace/ControlClient.h
@@ -36,9 +36,9 @@ public:
 	/** Returns true if the client is currently connected to a runtime */
     bool IsConnected() const;
 
-	/** Tells the runtime to attempt to record its trace event log.
-	 * @param Path Path that the runtime should send events to. */
-    void SendConnect(const TCHAR* Path);
+	/** Tells the runtime to attempt to record its trace event log to a socket.
+	 * @param Host Host IP address that the runtime should send events to. */
+    void SendSendTo(const TCHAR* Host);
 
 	/** Toggles events on and off. Events can be specified in two forms; either
 	 * by logger to address a logger's events as a whole (EventSpec="LoggerName")
diff --git a/Engine/Source/Developer/TraceServices/Private/SessionService.cpp b/Engine/Source/Developer/TraceServices/Private/SessionService.cpp
--- a/Engine/Source/Developer/TraceServices/Private/SessionService.cpp
+++ b/Engine/Source/Developer/TraceServices/Private/SessionService.cpp
@@ -217,7 +217,7 @@ bool FSessionService::ConnectSession(const TCHAR* ControlClientAddress)
 	{
 		return false;
 	}
-	ControlClient.SendConnect(*RecorderAddr->ToString(false));
+	ControlClient.SendSendTo(*RecorderAddr->ToString(false));
 	ControlClient.Disconnect();
 	return true;
 }
diff --git a/Engine/Source/Runtime/Launch/Private/LaunchEngineLoop.cpp b/Engine/Source/Runtime/Launch/Private/LaunchEngineLoop.cpp
--- a/Engine/Source/Runtime/Launch/Private/LaunchEngineLoop.cpp
+++ b/Engine/Source/Runtime/Launch/Private/LaunchEngineLoop.cpp
@@ -1262,12 +1262,12 @@ int32 FEngineLoop::PreInit(const TCHAR* CmdLine)
 		FString Parameter;
 		if (FParse::Value(CmdLine, TEXT("-tracehost="), Parameter))
 		{
-			Trace::Connect(*Parameter);
+			Trace::SendTo(*Parameter);
 		}
 
 		else if (FParse::Value(CmdLine, TEXT("-tracefile="), Parameter))
 		{
-			Trace::Open(*Parameter);
+			Trace::WriteTo(*Parameter);
 		}
 
 #if PLATFORM_WINDOWS && !UE_BUILD_SHIPPING
@@ -1277,7 +1277,7 @@ int32 FEngineLoop::PreInit(const TCHAR* CmdLine)
 			HANDLE KnownEvent = ::OpenEvent(EVENT_ALL_ACCESS, false, TEXT("Local\\UnrealInsightsRecorder"));
 			if (KnownEvent != nullptr)
 			{
-				Trace::Connect(TEXT("127.0.0.1"));
+				Trace::SendTo(TEXT("127.0.0.1"));
 				::CloseHandle(KnownEvent);
 			}
 		}
diff --git a/Engine/Source/Runtime/TraceLog/Private/Trace/Trace.cpp b/Engine/Source/Runtime/TraceLog/Private/Trace/Trace.cpp
--- a/Engine/Source/Runtime/TraceLog/Private/Trace/Trace.cpp
+++ b/Engine/Source/Runtime/TraceLog/Private/Trace/Trace.cpp
@@ -12,8 +12,8 @@ namespace Private
 {
 
 ////////////////////////////////////////////////////////////////////////////////
-bool	Writer_Connect(const ANSICHAR*);
-bool	Writer_Open(const ANSICHAR*);
+bool	Writer_SendTo(const ANSICHAR*);
+bool	Writer_WriteTo(const ANSICHAR*);
 uint32	Writer_EventToggle(const ANSICHAR*, bool);
 
 } // namespace Private
@@ -35,20 +35,20 @@ static void ToAnsiCheap(ANSICHAR (&Dest)[DestSize], const WIDECHAR* Src)
 };
 
 ////////////////////////////////////////////////////////////////////////////////
-bool Connect(const TCHAR* InHost)
+bool SendTo(const TCHAR* InHost)
 {
 	char Host[32];
 	ToAnsiCheap(Host, InHost);
 
-	return Private::Writer_Connect(Host);
+	return Private::Writer_WriteTo(Host);
 }
 
 ////////////////////////////////////////////////////////////////////////////////
-bool Open(const TCHAR* InPath)
+bool WriteTo(const TCHAR* InPath)
 {
 	char Path[512];
 	ToAnsiCheap(Path, InPath);
-	return Private::Writer_Open(Path);
+	return Private::Writer_WriteTo(Path);
 }
 
 ////////////////////////////////////////////////////////////////////////////////
diff --git a/Engine/Source/Runtime/TraceLog/Private/Trace/Writer.cpp b/Engine/Source/Runtime/TraceLog/Private/Trace/Writer.cpp
--- a/Engine/Source/Runtime/TraceLog/Private/Trace/Writer.cpp
+++ b/Engine/Source/Runtime/TraceLog/Private/Trace/Writer.cpp
@@ -530,8 +530,8 @@ struct FControlCommands
 };
 
 ////////////////////////////////////////////////////////////////////////////////
-bool	Writer_Connect(const ANSICHAR*);
-bool	Writer_Open(const ANSICHAR*);
+bool	Writer_SendTo(const ANSICHAR*);
+bool	Writer_WriteTo(const ANSICHAR*);
 uint32	Writer_EventToggle(const ANSICHAR*, bool);
 
 ////////////////////////////////////////////////////////////////////////////////
@@ -750,22 +750,22 @@ static void Writer_InitializeControl()
 	return;
 #endif
 
-	Writer_ControlAddCommand("Connect", nullptr,
+	Writer_ControlAddCommand("SendTo", nullptr,
 		[] (void*, uint32 ArgC, ANSICHAR const* const* ArgV)
 		{
 			if (ArgC > 0)
 			{
-				Writer_Connect(ArgV[0]);
+				Writer_SendTo(ArgV[0]);
 			}
 		}
 	);
 
-	Writer_ControlAddCommand("Open", nullptr,
+	Writer_ControlAddCommand("WriteTo", nullptr,
 		[] (void*, uint32 ArgC, ANSICHAR const* const* ArgV)
 		{
 			if (ArgC > 0)
 			{
-				Writer_Open(ArgV[0]);
+				Writer_WriteTo(ArgV[0]);
 			}
 		}
 	);
@@ -899,7 +899,7 @@ void Writer_Initialize()
 
 
 ////////////////////////////////////////////////////////////////////////////////
-bool Writer_Connect(const ANSICHAR* Host)
+bool Writer_SendTo(const ANSICHAR* Host)
 {
 	Writer_Initialize();
 
@@ -914,7 +914,7 @@ bool Writer_Connect(const ANSICHAR* Host)
 }
 
 ////////////////////////////////////////////////////////////////////////////////
-bool Writer_Open(const ANSICHAR* Path)
+bool Writer_WriteTo(const ANSICHAR* Path)
 {
 	Writer_Initialize();
 
diff --git a/Engine/Source/Runtime/TraceLog/Public/Trace/Trace.h b/Engine/Source/Runtime/TraceLog/Public/Trace/Trace.h
--- a/Engine/Source/Runtime/TraceLog/Public/Trace/Trace.h
+++ b/Engine/Source/Runtime/TraceLog/Public/Trace/Trace.h
@@ -9,8 +9,8 @@
 namespace Trace
 {
 
-UE_TRACE_API bool	Connect(const TCHAR* Host) UE_TRACE_IMPL(false);
-UE_TRACE_API bool	Open(const TCHAR* Path) UE_TRACE_IMPL(false);
+UE_TRACE_API bool	SendTo(const TCHAR* Host) UE_TRACE_IMPL(false);
+UE_TRACE_API bool	WriteTo(const TCHAR* Path) UE_TRACE_IMPL(false);
 UE_TRACE_API uint32 ToggleEvent(const TCHAR* Wildcard, bool bState) UE_TRACE_IMPL(false);
 
 } // namespace Trace
