Removed unused constant

diff --git a/Engine/Source/Runtime/TraceLog/Public/Trace/Detail/EventDef.h b/Engine/Source/Runtime/TraceLog/Public/Trace/Detail/EventDef.h
--- a/Engine/Source/Runtime/TraceLog/Public/Trace/Detail/EventDef.h
+++ b/Engine/Source/Runtime/TraceLog/Public/Trace/Detail/EventDef.h
@@ -51,8 +51,6 @@ struct FLiteralName;
 class FEventDef
 {
 public:
-	enum { HeaderSize = sizeof(uint32), }; // I'll just dump this here for now.
-
 	enum
 	{
 		Flag_Always		= 1 << 0,
