Asking trace to send its events to a socket was accidentally writing it to the local file system instead.

diff --git a/Engine/Source/Runtime/TraceLog/Private/Trace/Trace.cpp b/Engine/Source/Runtime/TraceLog/Private/Trace/Trace.cpp
--- a/Engine/Source/Runtime/TraceLog/Private/Trace/Trace.cpp
+++ b/Engine/Source/Runtime/TraceLog/Private/Trace/Trace.cpp
@@ -39,8 +39,7 @@ bool SendTo(const TCHAR* InHost)
 {
 	char Host[32];
 	ToAnsiCheap(Host, InHost);
-
-	return Private::Writer_WriteTo(Host);
+	return Private::Writer_SendTo(Host);
 }
 
 ////////////////////////////////////////////////////////////////////////////////
