UnixTrace.cpp was in the wrong place.

#rnx

diff --git a/Engine/Source/Runtime/TraceLog/Private/Trace/Detail/Unix/UnixTrace.cpp b/Engine/Source/Runtime/TraceLog/Private/Trace/Detail/Unix/UnixTrace.cpp
new file mode 100644
--- /dev/null
+++ b/Engine/Source/Runtime/TraceLog/Private/Trace/Detail/Unix/UnixTrace.cpp
@@ -0,0 +1,279 @@
+// Copyright 1998-2019 Epic Games, Inc. All Rights Reserved.
+
+#include "Trace/Trace.h"
+
+#if UE_TRACE_ENABLED
+
+#include <errno.h>
+#include <fcntl.h>
+#include <netdb.h>
+#include <pthread.h>
+#include <sys/mman.h>
+#include <sys/socket.h>
+#include <sys/types.h>
+#include <time.h>
+#include <unistd.h>
+#if defined(_GNU_SOURCE)
+	#include <sys/syscall.h>
+#endif // _GNU_SOURCE
+
+namespace Trace {
+namespace Private {
+
+////////////////////////////////////////////////////////////////////////////////
+uint8* MemoryReserve(SIZE_T Size)
+{
+	void* Ptr = mmap(nullptr, Size, PROT_READ | PROT_WRITE, MAP_PRIVATE | MAP_ANON, -1, 0);
+	return (Ptr != MAP_FAILED) ? reinterpret_cast<uint8*>(Ptr) : nullptr;
+}
+
+////////////////////////////////////////////////////////////////////////////////
+void MemoryFree(void* Address, SIZE_T Size)
+{
+	munmap(Address, Size);
+}
+
+////////////////////////////////////////////////////////////////////////////////
+void MemoryMap(void* Address, SIZE_T Size)
+{
+	// no-op if mmap()ed with R/W
+}
+
+////////////////////////////////////////////////////////////////////////////////
+void MemoryUnmap(void* Address, SIZE_T Size)
+{
+	madvise(Address, Size, MADV_DONTNEED);
+}
+
+
+
+////////////////////////////////////////////////////////////////////////////////
+UPTRINT ThreadCreate(const ANSICHAR* Name, void (*Entry)())
+{
+	void* (*PthreadThunk)(void*) = [] (void* Param) -> void * {
+		typedef void (*EntryType)(void);
+		(EntryType(Param))();
+		return nullptr;
+	};
+
+	pthread_t ThreadHandle;
+	if (pthread_create(&ThreadHandle, nullptr, PthreadThunk, reinterpret_cast<void *>(Entry)) != 0)
+	{
+		return 0;
+	}
+	return static_cast<UPTRINT>(ThreadHandle);
+}
+
+////////////////////////////////////////////////////////////////////////////////
+void ThreadSleep(uint32 Milliseconds)
+{
+	usleep(Milliseconds * 1000U);
+}
+
+////////////////////////////////////////////////////////////////////////////////
+void ThreadJoin(UPTRINT Handle)
+{
+	pthread_join(static_cast<pthread_t>(Handle), nullptr);
+}
+
+////////////////////////////////////////////////////////////////////////////////
+void ThreadDestroy(UPTRINT Handle)
+{
+	// no-op
+}
+
+
+
+////////////////////////////////////////////////////////////////////////////////
+uint64 TimeGetFrequency()
+{
+	return 1000000ull;
+}
+
+////////////////////////////////////////////////////////////////////////////////
+uint64 TimeGetTimestamp()
+{
+	// should stay in sync with FPlatformTime::Cycles64() or the timeline will be broken!
+	struct timespec TimeSpec;
+	clock_gettime(CLOCK_MONOTONIC, &TimeSpec);
+	return static_cast<uint64>(static_cast<uint64>(TimeSpec.tv_sec) * 1000000ULL + static_cast<uint64>(TimeSpec.tv_nsec) / 1000ULL);
+}
+
+
+
+////////////////////////////////////////////////////////////////////////////////
+static bool TcpSocketSetNonBlocking(int Socket, bool bNonBlocking)
+{
+	int Flags = fcntl(Socket, F_GETFL, 0);
+	if (Flags == -1)
+	{
+		return false;
+	}
+
+	Flags = bNonBlocking ? (Flags|O_NONBLOCK) : (Flags & ~O_NONBLOCK);
+	return fcntl(Socket, F_SETFL, Flags) >= 0;
+}
+
+////////////////////////////////////////////////////////////////////////////////
+UPTRINT TcpSocketConnect(const ANSICHAR* Host, uint16 Port)
+{
+	struct FAddrInfoPtr
+	{
+					~FAddrInfoPtr()	{ freeaddrinfo(Value); }
+		addrinfo*	operator -> ()	{ return Value; }
+		addrinfo**	operator & ()	{ return &Value; }
+		addrinfo*	Value;
+	};
+
+	FAddrInfoPtr Info;
+	addrinfo Hints = {};
+	Hints.ai_family = AF_INET;
+	Hints.ai_socktype = SOCK_STREAM;
+	Hints.ai_protocol = IPPROTO_TCP;
+	if (getaddrinfo(Host, nullptr, &Hints, &Info))
+	{
+		return 0;
+	}
+
+	if (&Info == nullptr)
+	{
+		return 0;
+	}
+
+	auto* SockAddr = (sockaddr_in*)Info->ai_addr;
+	SockAddr->sin_port = htons(Port);
+
+	int Socket = socket(AF_INET, SOCK_STREAM, IPPROTO_TCP);
+	if (Socket < 0)
+	{
+		return 0;
+	}
+
+	int Result = connect(Socket, Info->ai_addr, int(Info->ai_addrlen));
+	if (Result < 0)
+	{
+		close(Socket);
+		return 0;
+	}
+
+	if (!TcpSocketSetNonBlocking(Socket, false))
+	{
+		close(Socket);
+		return 0;
+	}
+
+	return UPTRINT(Socket + 1);
+}
+
+////////////////////////////////////////////////////////////////////////////////
+UPTRINT TcpSocketListen(uint16 Port)
+{
+	int Socket = socket(AF_INET, SOCK_STREAM, IPPROTO_TCP);
+	if (Socket < 0)
+	{
+		return 0;
+	}
+
+	sockaddr_in SockAddr;
+	SockAddr.sin_family = AF_INET;
+	SockAddr.sin_addr.s_addr = 0;
+	SockAddr.sin_port = htons(Port);
+	int Result = bind(Socket, reinterpret_cast<sockaddr*>(&SockAddr), sizeof(SockAddr));
+	if (Result < 0)
+	{
+		close(Socket);
+		return 0;
+	}
+
+	Result = listen(Socket, 1);
+	if (Result < 0)
+	{
+		close(Socket);
+		return 0;
+	}
+
+	if (!TcpSocketSetNonBlocking(Socket, true))
+	{
+		close(Socket);
+		return 0;
+	}
+
+	return UPTRINT(Socket + 1);
+}
+
+////////////////////////////////////////////////////////////////////////////////
+int32 TcpSocketAccept(UPTRINT Socket, UPTRINT& Out)
+{
+	int Inner = Socket - 1;
+
+	Inner = accept(Inner, nullptr, nullptr);
+	if (Inner < 0)
+	{
+		return (Inner == EAGAIN || Inner == EWOULDBLOCK) - 1; // 0 if would block else -1
+	}
+
+	if (!TcpSocketSetNonBlocking(Inner, false))
+	{
+		close(Inner);
+		return 0;
+	}
+
+	Out = UPTRINT(Inner + 1);
+	return 1;
+}
+
+////////////////////////////////////////////////////////////////////////////////
+bool TcpSocketHasData(UPTRINT Socket)
+{
+	int Inner = Socket - 1;
+	fd_set FdSet;
+	FD_ZERO(&FdSet);
+	FD_SET(Inner, &FdSet);
+	timeval TimeVal = {};
+	return (select(Inner + 1, &FdSet, nullptr, nullptr, &TimeVal) != 0);
+}
+
+
+
+////////////////////////////////////////////////////////////////////////////////
+bool IoWrite(UPTRINT Handle, const void* Data, uint32 Size)
+{
+	int Inner = int(Handle) - 1;
+	return write(Inner, Data, Size) == Size;
+}
+
+////////////////////////////////////////////////////////////////////////////////
+int32 IoRead(UPTRINT Handle, void* Data, uint32 Size)
+{
+	int Inner = int(Handle) - 1;
+	return read(Inner, Data, Size);
+}
+
+////////////////////////////////////////////////////////////////////////////////
+void IoClose(UPTRINT Handle)
+{
+	int Inner = int(Handle) - 1;
+	close(Inner);
+}
+
+
+
+////////////////////////////////////////////////////////////////////////////////
+UPTRINT FileOpen(const ANSICHAR* Path)
+{
+	int Flags = O_CREAT|O_WRONLY|O_TRUNC;
+	int Mode = S_IRUSR|S_IWUSR|S_IRGRP|S_IWGRP|S_IROTH|S_IWOTH;
+
+	int Out = open(Path, Flags, Mode);
+	if (Out < 0)
+	{
+		return 0;
+	}
+
+	return UPTRINT(Out + 1);
+}
+
+} // namespace Private
+} // namespace Trace
+
+#endif // UE_TRACE_ENABLED
diff --git a/Engine/Source/Runtime/TraceLog/Private/Trace/Unix/UnixTrace.cpp b/Engine/Source/Runtime/TraceLog/Private/Trace/Unix/UnixTrace.cpp
deleted file mode 100644
