Fixed 32 bit platforms not detecting that the default trace event buffer is being empty.

#rnx

diff --git a/Engine/Source/Runtime/TraceLog/Private/Trace/Writer.cpp b/Engine/Source/Runtime/TraceLog/Private/Trace/Writer.cpp
--- a/Engine/Source/Runtime/TraceLog/Private/Trace/Writer.cpp
+++ b/Engine/Source/Runtime/TraceLog/Private/Trace/Writer.cpp
@@ -213,6 +213,9 @@ static void Writer_InitializeBuffers()
 
 	static_assert(GPoolPageGrowth >= 0x10000, "Page growth must be >= 64KB");
 	static_assert(GPoolInitPageSize >= 0x10000, "Initial page size must be >= 64KB");
+
+	auto* EmptyBuffer = (FWriteBuffer*)GEmptyBuffer;
+	EmptyBuffer->Cursor = EmptyBuffer->Data;
 }
 
 ////////////////////////////////////////////////////////////////////////////////
