Early out if there's an existing connection.

#rnx

diff --git a/Engine/Source/Runtime/TraceLog/Private/Trace/Writer.cpp b/Engine/Source/Runtime/TraceLog/Private/Trace/Writer.cpp
--- a/Engine/Source/Runtime/TraceLog/Private/Trace/Writer.cpp
+++ b/Engine/Source/Runtime/TraceLog/Private/Trace/Writer.cpp
@@ -466,43 +466,43 @@ static void Writer_UpdateData()
 {
 	if (GPendingDataHandle)
 	{
+		// Reject the pending connection if we've already got a connection
 		if (GDataHandle)
 		{
 			IoClose(GPendingDataHandle);
 			GPendingDataHandle = 0;
+			return;
 		}
-		else
-		{
-			GDataHandle = GPendingDataHandle;
-			GPendingDataHandle = 0;
 
-			// Handshake.
-			const uint32 Magic = 'TRCE';
-			bool bOk = IoWrite(GDataHandle, &Magic, sizeof(Magic));
+		GDataHandle = GPendingDataHandle;
+		GPendingDataHandle = 0;
 
-			// Stream header
-			const struct {
-				uint8 Format;
-				uint8 Parameter;
-			} TransportHeader = { 2 };
-			bOk &= IoWrite(GDataHandle, &TransportHeader, sizeof(TransportHeader));
+		// Handshake.
+		const uint32 Magic = 'TRCE';
+		bool bOk = IoWrite(GDataHandle, &Magic, sizeof(Magic));
 
-			// Passively collected data
-			if (GHoldBuffer->GetSize())
-			{
-				bOk &= IoWrite(GDataHandle, GHoldBuffer->GetData(), GHoldBuffer->GetSize());
-			}
+		// Stream header
+		const struct {
+			uint8 Format;
+			uint8 Parameter;
+		} TransportHeader = { 2 };
+		bOk &= IoWrite(GDataHandle, &TransportHeader, sizeof(TransportHeader));
 
-			if (bOk)
-			{
-				GDataState = EDataState::Sending;
-				GHoldBuffer->Shutdown();
-			}
-			else
-			{
-				IoClose(GDataHandle);
-				GDataHandle = 0;
-			}
+		// Passively collected data
+		if (GHoldBuffer->GetSize())
+		{
+			bOk &= IoWrite(GDataHandle, GHoldBuffer->GetData(), GHoldBuffer->GetSize());
+		}
+
+		if (bOk)
+		{
+			GDataState = EDataState::Sending;
+			GHoldBuffer->Shutdown();
+		}
+		else
+		{
+			IoClose(GDataHandle);
+			GDataHandle = 0;
 		}
 	}
 
