In keeping with other platforms; only resolve host names on desktops.

#rnx

diff --git a/Engine/Source/Runtime/TraceLog/Private/Trace/Detail/Apple/AppleTrace.cpp b/Engine/Source/Runtime/TraceLog/Private/Trace/Detail/Apple/AppleTrace.cpp
--- a/Engine/Source/Runtime/TraceLog/Private/Trace/Detail/Apple/AppleTrace.cpp
+++ b/Engine/Source/Runtime/TraceLog/Private/Trace/Detail/Apple/AppleTrace.cpp
@@ -4,6 +4,7 @@
 
 #if UE_TRACE_ENABLED
 
+#include <arpa/inet.h>
 #include <fcntl.h>
 #include <mach/mach.h>
 #include <mach/mach_time.h>
@@ -126,6 +127,7 @@ static bool TcpSocketSetNonBlocking(int Socket, bool bNonBlocking)
 ////////////////////////////////////////////////////////////////////////////////
 UPTRINT TcpSocketConnect(const ANSICHAR* Host, uint16 Port)
 {
+#if PLATFORM_MAC // We're only accepting named hosts on desktop platforms
 	struct FAddrInfoPtr
 	{
 					~FAddrInfoPtr()	{ freeaddrinfo(Value); }
@@ -151,6 +153,16 @@ UPTRINT TcpSocketConnect(const ANSICHAR* Host, uint16 Port)
 
 	auto* SockAddr = (sockaddr_in*)Info->ai_addr;
 	SockAddr->sin_port = htons(Port);
+	int SockAddrSize = int(Info->ai_addrlen);
+#else
+	sockaddr_in SockAddrIp;
+	SockAddrIp.sin_family = AF_INET;
+	SockAddrIp.sin_addr.s_addr = inet_addr(Host);
+	SockAddrIp.sin_port = htons(Port);
+
+	auto* SockAddr = &SockAddrIp;
+	int SockAddrSize = sizeof(SockAddrIp);
+#endif
 
 	int Socket = socket(AF_INET, SOCK_STREAM, IPPROTO_TCP);
 	if (Socket < 0)
@@ -158,7 +170,7 @@ UPTRINT TcpSocketConnect(const ANSICHAR* Host, uint16 Port)
 		return 0;
 	}
 
-	int Result = connect(Socket, Info->ai_addr, int(Info->ai_addrlen));
+	int Result = connect(Socket, (sockaddr*)SockAddr, SockAddrSize);
 	if (Result < 0)
 	{
 		close(Socket);
