TimeGetFrequency() is not intended to be called often so it can be expensive. Function-scoped statics are expensive anyway as C++11 states they must be thread safe.

#rnx

diff --git a/Engine/Source/Runtime/TraceLog/Private/Trace/Detail/Apple/AppleTrace.cpp b/Engine/Source/Runtime/TraceLog/Private/Trace/Detail/Apple/AppleTrace.cpp
--- a/Engine/Source/Runtime/TraceLog/Private/Trace/Detail/Apple/AppleTrace.cpp
+++ b/Engine/Source/Runtime/TraceLog/Private/Trace/Detail/Apple/AppleTrace.cpp
@@ -83,23 +83,11 @@ void ThreadDestroy(UPTRINT Handle)
 
 
 
-////////////////////////////////////////////////////////////////////////////////
-const mach_timebase_info_data_t& TimeGetInfo()
-{
-    static mach_timebase_info_data_t Info;
-    static dispatch_once_t OnceToken;
-
-	dispatch_once(&OnceToken, ^{
-		   (void) mach_timebase_info(&Info);
-	   });
-
-	return Info;
-}
-
 ////////////////////////////////////////////////////////////////////////////////
 uint64 TimeGetFrequency()
 {
-	const mach_timebase_info_data_t& Info = TimeGetInfo();
+	mach_timebase_info_data_t Info;
+	mach_timebase_info(&Info);
 	return (uint64(1 * 1000 * 1000 * 1000) * uint64(Info.denom)) / uint64(Info.numer);
 }
 
