Exposed the maximum event size to users of Trace so they can plan their attachments appropriately

#rnx

diff --git a/Engine/Source/Runtime/TraceLog/Private/Trace/Trace.cpp b/Engine/Source/Runtime/TraceLog/Private/Trace/Trace.cpp
--- a/Engine/Source/Runtime/TraceLog/Private/Trace/Trace.cpp
+++ b/Engine/Source/Runtime/TraceLog/Private/Trace/Trace.cpp
@@ -12,6 +12,7 @@ namespace Private
 {
 
 ////////////////////////////////////////////////////////////////////////////////
+uint32	Writer_GetMaxEventSize();
 bool	Writer_SendTo(const ANSICHAR*);
 bool	Writer_WriteTo(const ANSICHAR*);
 uint32	Writer_EventToggle(const ANSICHAR*, bool);
@@ -34,6 +35,12 @@ static void ToAnsiCheap(ANSICHAR (&Dest)[DestSize], const WIDECHAR* Src)
 	}
 };
 
+////////////////////////////////////////////////////////////////////////////////
+uint32 GetMaxEventSize()
+{
+	return Private::Writer_GetMaxEventSize();
+}
+
 ////////////////////////////////////////////////////////////////////////////////
 bool SendTo(const TCHAR* InHost)
 {
diff --git a/Engine/Source/Runtime/TraceLog/Private/Trace/Writer.cpp b/Engine/Source/Runtime/TraceLog/Private/Trace/Writer.cpp
--- a/Engine/Source/Runtime/TraceLog/Private/Trace/Writer.cpp
+++ b/Engine/Source/Runtime/TraceLog/Private/Trace/Writer.cpp
@@ -73,6 +73,13 @@ T_ALIGN static uint8* volatile		GPoolPageCursor;	// = nullptr;
 T_ALIGN static void* volatile		GPoolFreeList;		// = nullptr;
 #undef T_ALIGN
 
+////////////////////////////////////////////////////////////////////////////////
+uint32 Writer_GetMaxEventSize()
+{
+	// The bais is to allow for some overhead. Its value was chosen arbitrarily.
+	return GPoolBlockSize - 512;
+}
+
 ////////////////////////////////////////////////////////////////////////////////
 #if !IS_MONOLITHIC
 UE_TRACE_API FWriteBuffer* Writer_GetBuffer()
diff --git a/Engine/Source/Runtime/TraceLog/Public/Trace/Trace.h b/Engine/Source/Runtime/TraceLog/Public/Trace/Trace.h
--- a/Engine/Source/Runtime/TraceLog/Public/Trace/Trace.h
+++ b/Engine/Source/Runtime/TraceLog/Public/Trace/Trace.h
@@ -9,9 +9,10 @@
 namespace Trace
 {
 
+UE_TRACE_API uint32	GetMaxEventSize() UE_TRACE_IMPL(0);
 UE_TRACE_API bool	SendTo(const TCHAR* Host) UE_TRACE_IMPL(false);
 UE_TRACE_API bool	WriteTo(const TCHAR* Path) UE_TRACE_IMPL(false);
-UE_TRACE_API uint32 ToggleEvent(const TCHAR* Wildcard, bool bState) UE_TRACE_IMPL(false);
+UE_TRACE_API uint32 ToggleEvent(const TCHAR* Wildcard, bool bState) UE_TRACE_IMPL(0);
 
 } // namespace Trace
 
