Removed unused PublicSystemLibraries statement.

#rnx none

diff --git a/Engine/Source/Runtime/TraceLog/TraceLog.Build.cs b/Engine/Source/Runtime/TraceLog/TraceLog.Build.cs
--- a/Engine/Source/Runtime/TraceLog/TraceLog.Build.cs
+++ b/Engine/Source/Runtime/TraceLog/TraceLog.Build.cs
@@ -10,10 +10,5 @@ public class TraceLog : ModuleRules
 	{
 		bRequiresImplementModule = false;
 		PublicIncludePathModuleNames.Add("Core");
-
-        if (Target.Platform == UnrealTargetPlatform.PS4)
-        {
-            PublicSystemLibraries.Add("SceNet_stub_weak");
-        }
     }
 }
