Changed includes to reduced contact surface with Trace. This should hopefully improve my iteration times as I will no longer need to rebuild everything when touching public Trace headers.

#rnx

diff --git a/Engine/Source/Runtime/Core/Private/ProfilingDebugging/CountersTrace.cpp b/Engine/Source/Runtime/Core/Private/ProfilingDebugging/CountersTrace.cpp
--- a/Engine/Source/Runtime/Core/Private/ProfilingDebugging/CountersTrace.cpp
+++ b/Engine/Source/Runtime/Core/Private/ProfilingDebugging/CountersTrace.cpp
@@ -2,6 +2,7 @@
 #include "ProfilingDebugging/CountersTrace.h"
 #include "HAL/PlatformTime.h"
 #include "Misc/Parse.h"
+#include "Trace/Trace.h"
 
 #if COUNTERSTRACE_ENABLED
 
diff --git a/Engine/Source/Runtime/Core/Public/ProfilingDebugging/CountersTrace.h b/Engine/Source/Runtime/Core/Public/ProfilingDebugging/CountersTrace.h
--- a/Engine/Source/Runtime/Core/Public/ProfilingDebugging/CountersTrace.h
+++ b/Engine/Source/Runtime/Core/Public/ProfilingDebugging/CountersTrace.h
@@ -3,7 +3,7 @@
 #pragma once
 
 #include "CoreTypes.h"
-#include "Trace/Trace.h"
+#include "Trace/Config.h"
 
 #if UE_TRACE_ENABLED && !UE_BUILD_SHIPPING
 #define COUNTERSTRACE_ENABLED 1
diff --git a/Engine/Source/Runtime/Core/Public/ProfilingDebugging/CsvProfilerTrace.h b/Engine/Source/Runtime/Core/Public/ProfilingDebugging/CsvProfilerTrace.h
--- a/Engine/Source/Runtime/Core/Public/ProfilingDebugging/CsvProfilerTrace.h
+++ b/Engine/Source/Runtime/Core/Public/ProfilingDebugging/CsvProfilerTrace.h
@@ -3,7 +3,7 @@
 #pragma once
 
 #include "CoreTypes.h"
-#include "Trace/Trace.h"
+#include "Trace/Config.h"
 
 #if UE_TRACE_ENABLED && !UE_BUILD_SHIPPING
 #define CSVPROFILERTRACE_ENABLED WITH_ENGINE
diff --git a/Engine/Source/Runtime/Core/Public/ProfilingDebugging/LoadTimeTracker.h b/Engine/Source/Runtime/Core/Public/ProfilingDebugging/LoadTimeTracker.h
--- a/Engine/Source/Runtime/Core/Public/ProfilingDebugging/LoadTimeTracker.h
+++ b/Engine/Source/Runtime/Core/Public/ProfilingDebugging/LoadTimeTracker.h
@@ -11,7 +11,7 @@
 #include "Containers/Map.h"
 #include "UObject/NameTypes.h"
 #include "ProfilingDebugging/ScopedTimers.h"
-#include "Trace/Trace.h"
+#include "Trace/Config.h"
 
 #ifndef ENABLE_LOADTIME_TRACKING
 	#define ENABLE_LOADTIME_TRACKING 0
diff --git a/Engine/Source/Runtime/Core/Public/ProfilingDebugging/PlatformFileTrace.h b/Engine/Source/Runtime/Core/Public/ProfilingDebugging/PlatformFileTrace.h
--- a/Engine/Source/Runtime/Core/Public/ProfilingDebugging/PlatformFileTrace.h
+++ b/Engine/Source/Runtime/Core/Public/ProfilingDebugging/PlatformFileTrace.h
@@ -3,7 +3,7 @@
 #pragma once
 
 #include "CoreTypes.h"
-#include "Trace/Trace.h"
+#include "Trace/Config.h"
 
 #if UE_TRACE_ENABLED && PLATFORM_WINDOWS && !UE_BUILD_SHIPPING
 #define PLATFORMFILETRACE_ENABLED 1
diff --git a/Engine/Source/Runtime/Core/Public/Serialization/LoadTimeTrace.h b/Engine/Source/Runtime/Core/Public/Serialization/LoadTimeTrace.h
--- a/Engine/Source/Runtime/Core/Public/Serialization/LoadTimeTrace.h
+++ b/Engine/Source/Runtime/Core/Public/Serialization/LoadTimeTrace.h
@@ -3,7 +3,7 @@
 #pragma once
 
 #include "CoreTypes.h"
-#include "Trace/Trace.h"
+#include "Trace/Config.h"
 #include "ProfilingDebugging/FormatArgsTrace.h"
 
 #if UE_TRACE_ENABLED && !UE_BUILD_SHIPPING
diff --git a/Engine/Source/Runtime/Engine/Public/ObjectTrace.h b/Engine/Source/Runtime/Engine/Public/ObjectTrace.h
--- a/Engine/Source/Runtime/Engine/Public/ObjectTrace.h
+++ b/Engine/Source/Runtime/Engine/Public/ObjectTrace.h
@@ -3,6 +3,7 @@
 #pragma once
 
 #include "CoreTypes.h"
+#include "Trace/Config.h"
 
 #if !IS_PROGRAM && !UE_BUILD_SHIPPING
 #define OBJECT_TRACE_ENABLED 1
diff --git a/Engine/Source/Runtime/RHI/Public/GpuProfilerTrace.h b/Engine/Source/Runtime/RHI/Public/GpuProfilerTrace.h
--- a/Engine/Source/Runtime/RHI/Public/GpuProfilerTrace.h
+++ b/Engine/Source/Runtime/RHI/Public/GpuProfilerTrace.h
@@ -3,7 +3,7 @@
 #pragma once
 
 #include "CoreTypes.h"
-#include "Trace/Trace.h"
+#include "Trace/Config.h"
 
 #if UE_TRACE_ENABLED && !UE_BUILD_SHIPPING
 #define GPUPROFILERTRACE_ENABLED 0
diff --git a/Engine/Source/Runtime/TraceLog/Private/Trace/Detail/Android/AndroidTrace.cpp b/Engine/Source/Runtime/TraceLog/Private/Trace/Detail/Android/AndroidTrace.cpp
--- a/Engine/Source/Runtime/TraceLog/Private/Trace/Detail/Android/AndroidTrace.cpp
+++ b/Engine/Source/Runtime/TraceLog/Private/Trace/Detail/Android/AndroidTrace.cpp
@@ -1,6 +1,6 @@
 // Copyright 1998-2019 Epic Games, Inc. All Rights Reserved.
 
-#include "Trace/Trace.h"
+#include "Trace/Config.h"
 
 #if UE_TRACE_ENABLED
 
diff --git a/Engine/Source/Runtime/TraceLog/Private/Trace/Detail/Apple/AppleTrace.cpp b/Engine/Source/Runtime/TraceLog/Private/Trace/Detail/Apple/AppleTrace.cpp
--- a/Engine/Source/Runtime/TraceLog/Private/Trace/Detail/Apple/AppleTrace.cpp
+++ b/Engine/Source/Runtime/TraceLog/Private/Trace/Detail/Apple/AppleTrace.cpp
@@ -1,6 +1,6 @@
 // Copyright 1998-2019 Epic Games, Inc. All Rights Reserved.
 
-#include "Trace/Trace.h"
+#include "Trace/Config.h"
 
 #if UE_TRACE_ENABLED
 
diff --git a/Engine/Source/Runtime/TraceLog/Private/Trace/Detail/Unix/UnixTrace.cpp b/Engine/Source/Runtime/TraceLog/Private/Trace/Detail/Unix/UnixTrace.cpp
--- a/Engine/Source/Runtime/TraceLog/Private/Trace/Detail/Unix/UnixTrace.cpp
+++ b/Engine/Source/Runtime/TraceLog/Private/Trace/Detail/Unix/UnixTrace.cpp
@@ -1,6 +1,6 @@
 // Copyright 1998-2019 Epic Games, Inc. All Rights Reserved.
 
-#include "Trace/Trace.h"
+#include "Trace/Config.h"
 
 #if UE_TRACE_ENABLED
 
diff --git a/Engine/Source/Runtime/TraceLog/Private/Trace/Detail/Windows/WindowsTrace.cpp b/Engine/Source/Runtime/TraceLog/Private/Trace/Detail/Windows/WindowsTrace.cpp
--- a/Engine/Source/Runtime/TraceLog/Private/Trace/Detail/Windows/WindowsTrace.cpp
+++ b/Engine/Source/Runtime/TraceLog/Private/Trace/Detail/Windows/WindowsTrace.cpp
@@ -1,6 +1,6 @@
 // Copyright 1998-2019 Epic Games, Inc. All Rights Reserved.
 
-#include "Trace/Trace.h"
+#include "Trace/Config.h"
 
 #if UE_TRACE_ENABLED
 
diff --git a/Engine/Source/Runtime/TraceLog/Private/Trace/Trace.cpp b/Engine/Source/Runtime/TraceLog/Private/Trace/Trace.cpp
--- a/Engine/Source/Runtime/TraceLog/Private/Trace/Trace.cpp
+++ b/Engine/Source/Runtime/TraceLog/Private/Trace/Trace.cpp
@@ -1,10 +1,11 @@
 // Copyright 1998-2019 Epic Games, Inc. All Rights Reserved.
 
-#include "Trace/Trace.h"
-#include "Trace/Detail/EventDef.h"
+#include "Trace/Trace.h" // should be Config.h :(
 
 #if UE_TRACE_ENABLED
 
+#include "Trace/Detail/EventDef.h"
+
 namespace Trace
 {
 
diff --git a/Engine/Source/Runtime/TraceLog/Private/Trace/Writer.cpp b/Engine/Source/Runtime/TraceLog/Private/Trace/Writer.cpp
--- a/Engine/Source/Runtime/TraceLog/Private/Trace/Writer.cpp
+++ b/Engine/Source/Runtime/TraceLog/Private/Trace/Writer.cpp
@@ -1,11 +1,11 @@
 // Copyright 1998-2019 Epic Games, Inc. All Rights Reserved.
 
-#include "Trace/Detail/Trace.h"
+#include "Trace/Config.h"
 
 #if UE_TRACE_ENABLED
 
-#include "Trace/Platform.h"
 #include "Trace/Detail/Atomic.h"
+#include "Trace/Platform.h"
 #include "Trace/Trace.h"
 
 #include "Misc/CString.h"
