More include statement resuffling to minimize contact surface (and thus rebuilds).

#rnx

diff --git a/Engine/Source/Runtime/TraceLog/Private/Trace/Platform.h b/Engine/Source/Runtime/TraceLog/Private/Trace/Platform.h
--- a/Engine/Source/Runtime/TraceLog/Private/Trace/Platform.h
+++ b/Engine/Source/Runtime/TraceLog/Private/Trace/Platform.h
@@ -2,7 +2,7 @@
 
 #pragma once
 
-#include "Trace/Detail/Trace.h"
+#include "Trace/Config.h"
 
 #if UE_TRACE_ENABLED
 
diff --git a/Engine/Source/Runtime/TraceLog/Public/Trace/Detail/EventDef.h b/Engine/Source/Runtime/TraceLog/Public/Trace/Detail/EventDef.h
--- a/Engine/Source/Runtime/TraceLog/Public/Trace/Detail/EventDef.h
+++ b/Engine/Source/Runtime/TraceLog/Public/Trace/Detail/EventDef.h
@@ -2,7 +2,7 @@
 
 #pragma once
 
-#include "Trace.h"
+#include "Trace/Config.h"
 
 namespace Trace
 {
@@ -60,10 +60,10 @@ public:
 	class FLogScope
 	{
 	public:
-						FLogScope(uint16 EventUid, uint16 Size);
-						FLogScope(uint16 EventUid, uint16 Size, uint16 ExtraBytes);
-						~FLogScope();
-		uint8*			Ptr;
+								FLogScope(uint16 EventUid, uint16 Size);
+								FLogScope(uint16 EventUid, uint16 Size, uint16 ExtraBytes);
+								~FLogScope();
+		uint8*					Ptr;
 	};
 
 	void*						Handle;
@@ -80,7 +80,7 @@ public:
 		uint16					Test;
 	}							Enabled;
 	bool						bInitialized;
-	UE_TRACE_API static void	Create(FEventDef* Target, const FLiteralName& LoggerName, const FLiteralName& EventName, const FFieldDesc* FieldDescs, uint32 FieldCount, uint32 Flags=0);
+	TRACELOG_API static void	Create(FEventDef* Target, const FLiteralName& LoggerName, const FLiteralName& EventName, const FFieldDesc* FieldDescs, uint32 FieldCount, uint32 Flags=0);
 };
 
 } // namespace Trace
diff --git a/Engine/Source/Runtime/TraceLog/Public/Trace/Detail/Field.h b/Engine/Source/Runtime/TraceLog/Public/Trace/Detail/Field.h
--- a/Engine/Source/Runtime/TraceLog/Public/Trace/Detail/Field.h
+++ b/Engine/Source/Runtime/TraceLog/Public/Trace/Detail/Field.h
@@ -2,7 +2,7 @@
 
 #pragma once
 
-#include "Trace.h"
+#include "Trace/Config.h"
 
 #include "Templates/UnrealTemplate.h"
 
diff --git a/Engine/Source/Runtime/TraceLog/Public/Trace/Detail/Writer.inl b/Engine/Source/Runtime/TraceLog/Public/Trace/Detail/Writer.inl
--- a/Engine/Source/Runtime/TraceLog/Public/Trace/Detail/Writer.inl
+++ b/Engine/Source/Runtime/TraceLog/Public/Trace/Detail/Writer.inl
@@ -35,9 +35,9 @@ struct FWriteBuffer
 #endif
 
 ////////////////////////////////////////////////////////////////////////////////
-extern UE_TRACE_API void* volatile	GLastEvent;
-UE_TRACE_API uint8*					Writer_NextBuffer(uint16);
-UE_TRACE_API FWriteBuffer*			Writer_GetBuffer();
+extern TRACELOG_API void* volatile	GLastEvent;
+TRACELOG_API uint8*					Writer_NextBuffer(uint16);
+TRACELOG_API FWriteBuffer*			Writer_GetBuffer();
 
 #if IS_MONOLITHIC
 extern thread_local FWriteBuffer* GWriteBuffer;
