Use TRACELOG_API internally and UE_TRACE_API for the public facing stuff that might get stubbed out.

#rnx

diff --git a/Engine/Source/Runtime/TraceLog/Private/Trace/Writer.cpp b/Engine/Source/Runtime/TraceLog/Private/Trace/Writer.cpp
--- a/Engine/Source/Runtime/TraceLog/Private/Trace/Writer.cpp
+++ b/Engine/Source/Runtime/TraceLog/Private/Trace/Writer.cpp
@@ -67,7 +67,7 @@ void Writer_InitializeTiming()
 static uint8						GEmptyBuffer[sizeof(FWriteBuffer)];
 thread_local FWriteBuffer*			GWriteBuffer		= (FWriteBuffer*)GEmptyBuffer;
 T_ALIGN static void* volatile		GFirstEvent;
-T_ALIGN UE_TRACE_API void* volatile	GLastEvent;			// = nullptr;
+T_ALIGN TRACELOG_API void* volatile	GLastEvent;			// = nullptr;
 static const uint32					GPoolSize			= 384 << 20; // 384MB ought to be enough
 T_ALIGN static UPTRINT volatile		GThreadId;			// = 0;
 static const uint32					GPoolBlockSize		= 4 << 10;
@@ -87,7 +87,7 @@ uint32 Writer_GetMaxEventSize()
 
 ////////////////////////////////////////////////////////////////////////////////
 #if !IS_MONOLITHIC
-UE_TRACE_API FWriteBuffer* Writer_GetBuffer()
+TRACELOG_API FWriteBuffer* Writer_GetBuffer()
 {
 	// Thread locals and DLLs don't mix so for modular builds we are forced to
 	// export this function to access thread-local variables.
@@ -168,7 +168,7 @@ static FWriteBuffer* Writer_NextBufferInternal(uint32 PageGrowth)
 }
 
 ////////////////////////////////////////////////////////////////////////////////
-UE_TRACE_API uint8* Writer_NextBuffer(uint16 Size)
+TRACELOG_API uint8* Writer_NextBuffer(uint16 Size)
 {
 	FWriteBuffer* Current = GWriteBuffer;
 
diff --git a/Engine/Source/Runtime/TraceLog/Public/Trace/Detail/Trace.h b/Engine/Source/Runtime/TraceLog/Public/Trace/Detail/Trace.h
--- a/Engine/Source/Runtime/TraceLog/Public/Trace/Detail/Trace.h
+++ b/Engine/Source/Runtime/TraceLog/Public/Trace/Detail/Trace.h
@@ -5,17 +5,6 @@
 #include "CoreTypes.h"
 #include "Trace/Config.h"
 
-////////////////////////////////////////////////////////////////////////////////
-#if UE_TRACE_ENABLED
-#	define UE_TRACE_IMPL(...)
-#	define UE_TRACE_API			TRACELOG_API
-#else
-#	define UE_TRACE_IMPL(...)	{ return __VA_ARGS__; }
-#	define UE_TRACE_API			inline
-#endif
-
-
-
 #if UE_TRACE_ENABLED
 
 #define TRACE_PRIVATE_EVENT_DEFINE(LoggerName, EventName) \
diff --git a/Engine/Source/Runtime/TraceLog/Public/Trace/Trace.h b/Engine/Source/Runtime/TraceLog/Public/Trace/Trace.h
--- a/Engine/Source/Runtime/TraceLog/Public/Trace/Trace.h
+++ b/Engine/Source/Runtime/TraceLog/Public/Trace/Trace.h
@@ -5,6 +5,15 @@
 #include "Detail/Trace.h"
 #include "Detail/Field.h"
 
+////////////////////////////////////////////////////////////////////////////////
+#if UE_TRACE_ENABLED
+#	define UE_TRACE_IMPL(...)
+#	define UE_TRACE_API			TRACELOG_API
+#else
+#	define UE_TRACE_IMPL(...)	{ return __VA_ARGS__; }
+#	define UE_TRACE_API			inline
+#endif
+
 ////////////////////////////////////////////////////////////////////////////////
 namespace Trace
 {
