Moved the master protocol define into Protocol.h so we can guarantee it's defined correctly when Trace is on and the protocol has been included.

#rnx

diff --git a/Engine/Source/Runtime/TraceLog/Public/Trace/Detail/Protocol.h b/Engine/Source/Runtime/TraceLog/Public/Trace/Detail/Protocol.h
--- a/Engine/Source/Runtime/TraceLog/Public/Trace/Detail/Protocol.h
+++ b/Engine/Source/Runtime/TraceLog/Public/Trace/Detail/Protocol.h
@@ -2,6 +2,12 @@
 
 #pragma once
 
+#include "Trace/Config.h"
+
+#if UE_TRACE_ENABLED
+#	define TRACE_PRIVATE_PROTOCOL_0
+#endif
+
 #if defined(_MSC_VER)
 	#pragma warning(push)
 	#pragma warning(disable : 4200) // non-standard zero-sized array
diff --git a/Engine/Source/Runtime/TraceLog/Public/Trace/Detail/Trace.h b/Engine/Source/Runtime/TraceLog/Public/Trace/Detail/Trace.h
--- a/Engine/Source/Runtime/TraceLog/Public/Trace/Detail/Trace.h
+++ b/Engine/Source/Runtime/TraceLog/Public/Trace/Detail/Trace.h
@@ -2,12 +2,11 @@
 
 #pragma once
 
-#include "CoreTypes.h"
 #include "Trace/Config.h"
 
 #if UE_TRACE_ENABLED
 
-#define TRACE_PRIVATE_PROTOCOL_0
+#include "CoreTypes.h"
 
 #define TRACE_PRIVATE_EVENT_DEFINE(LoggerName, EventName) \
 	Trace::FEventDef LoggerName##EventName##Event;
