Corrected the trace transport version that I had bumped in error.

#rnx

diff --git a/Engine/Source/Runtime/TraceLog/Private/Trace/Writer.cpp b/Engine/Source/Runtime/TraceLog/Private/Trace/Writer.cpp
--- a/Engine/Source/Runtime/TraceLog/Private/Trace/Writer.cpp
+++ b/Engine/Source/Runtime/TraceLog/Private/Trace/Writer.cpp
@@ -509,7 +509,7 @@ static void Writer_UpdateData()
 
 		// Stream header
 		const struct {
-			uint8 TransportVersion	= 4;
+			uint8 TransportVersion	= 2;
 			uint8 ProtocolVersion	= 0;
 		} TransportHeader;
 		bOk &= IoWrite(GDataHandle, &TransportHeader, sizeof(TransportHeader));
