Write trace events forwards instead of backwards.

#rnx

diff --git a/Engine/Source/Runtime/TraceLog/Private/Trace/Writer.cpp b/Engine/Source/Runtime/TraceLog/Private/Trace/Writer.cpp
--- a/Engine/Source/Runtime/TraceLog/Private/Trace/Writer.cpp
+++ b/Engine/Source/Runtime/TraceLog/Private/Trace/Writer.cpp
@@ -135,7 +135,9 @@ static FWriteBuffer* Writer_NextBufferInternal(uint32 PageGrowth)
 		// it to the free list.
 		MemoryMap(PageBase, PageGrowth);
 
-		// The first block in the page we'll use for the next buffer
+		// The first block in the page we'll use for the next buffer. Note that the
+		// buffer objects are at the _end_ of their blocks.
+		PageBase += GPoolBlockSize - sizeof(FWriteBuffer);
 		Next = (FWriteBuffer*)PageBase;
 		uint8* FirstBlock = PageBase + GPoolBlockSize;
 
@@ -164,7 +166,7 @@ static FWriteBuffer* Writer_NextBufferInternal(uint32 PageGrowth)
 
 	GWriteBuffer = Next;
 
-	Next->Cursor = ((uint8*)Next + GPoolBlockSize);
+	Next->Cursor = ((uint8*)Next - GPoolBlockSize + sizeof(FWriteBuffer));
 	return Next;
 }
 
@@ -218,7 +220,7 @@ TRACELOG_API uint8* Writer_NextBuffer(uint16 Size)
 	FWriteBuffer* NextBuffer = Writer_NextBufferInternal(GPoolPageGrowth);
 	NextBuffer->ThreadId = ThreadId;
 
-	NextBuffer->Cursor -= Size;
+	NextBuffer->Cursor += Size;
 	return NextBuffer->Cursor;
 }
 
@@ -234,7 +236,7 @@ static void Writer_InitializeBuffers()
 	static_assert(GPoolInitPageSize >= 0x10000, "Initial page size must be >= 64KB");
 
 	auto* EmptyBuffer = (FWriteBuffer*)GEmptyBuffer;
-	EmptyBuffer->Cursor = EmptyBuffer->Data;
+	EmptyBuffer->Cursor = GEmptyBuffer;
 }
 
 ////////////////////////////////////////////////////////////////////////////////
@@ -457,7 +459,7 @@ static void Writer_ConsumeEvents()
 	for (void* EventPtr = LatestEvent; EventPtr != nullptr; )
 	{
 		// Is this "event" a retired buffer?
-		if ((UPTRINT(EventPtr) & (GPoolBlockSize - 1)) == 0)
+		if ((UPTRINT(EventPtr) & (GPoolBlockSize - 1)) == (GPoolBlockSize - sizeof(FWriteBuffer)))
 		{
 			auto* Retiree = (FWriteBuffer*)EventPtr;
 			EventPtr = *(void**)EventPtr;
diff --git a/Engine/Source/Runtime/TraceLog/Public/Trace/Detail/Writer.inl b/Engine/Source/Runtime/TraceLog/Public/Trace/Detail/Writer.inl
--- a/Engine/Source/Runtime/TraceLog/Public/Trace/Detail/Writer.inl
+++ b/Engine/Source/Runtime/TraceLog/Public/Trace/Detail/Writer.inl
@@ -27,7 +27,6 @@ struct FWriteBuffer
 		FWriteBuffer*	Next;
 	};
 	uint32				ThreadId;
-	uint8				Data[];
 };
 
 #if defined(_MSC_VER)
@@ -58,11 +57,12 @@ inline uint8* Writer_BeginLog(uint16 EventUid, uint16 Size)
 	uint32 AllocSize = ((Size + HeaderSize) + 7) & ~7;
 
 	FWriteBuffer* Buffer = Writer_GetBuffer();
-	uint8* Cursor = (Buffer->Cursor -= AllocSize);
-	if (UNLIKELY(PTRINT(Cursor) < PTRINT(Buffer->Data)))
+	uint8* Cursor = (Buffer->Cursor += AllocSize);
+	if (UNLIKELY(UPTRINT(Cursor) > UPTRINT(Buffer)))
 	{
 		Cursor = Writer_NextBuffer(AllocSize);
 	}
+	Cursor -= AllocSize;
 
 	uint32* Out = (uint32*)(Cursor + sizeof(void*));
 	Out[0] = (uint32(Size) << 16)|uint32(EventUid);
