Renamed "Next" to "NextBuffer" for clarity.

#rnx

diff --git a/Engine/Source/Runtime/TraceLog/Private/Trace/Writer.cpp b/Engine/Source/Runtime/TraceLog/Private/Trace/Writer.cpp
--- a/Engine/Source/Runtime/TraceLog/Private/Trace/Writer.cpp
+++ b/Engine/Source/Runtime/TraceLog/Private/Trace/Writer.cpp
@@ -100,7 +100,7 @@ TRACELOG_API FWriteBuffer* Writer_GetBuffer()
 static FWriteBuffer* Writer_NextBufferInternal(uint32 PageGrowth)
 {
 	// Fetch a new buffer
-	FWriteBuffer* Next;
+	FWriteBuffer* NextBuffer;
 	while (true)
 	{
 		// First we'll try one from the free list
@@ -117,7 +117,7 @@ static FWriteBuffer* Writer_NextBufferInternal(uint32 PageGrowth)
 		// If we didn't fetch the sentinal then we've taken a block we can use
 		if (Owned != nullptr)
 		{
-			Next = (FWriteBuffer*)Owned;
+			NextBuffer = (FWriteBuffer*)Owned;
 			break;
 		}
 
@@ -138,7 +138,7 @@ static FWriteBuffer* Writer_NextBufferInternal(uint32 PageGrowth)
 		// The first block in the page we'll use for the next buffer. Note that the
 		// buffer objects are at the _end_ of their blocks.
 		PageBase += GPoolBlockSize - sizeof(FWriteBuffer);
-		Next = (FWriteBuffer*)PageBase;
+		NextBuffer = (FWriteBuffer*)PageBase;
 		uint8* FirstBlock = PageBase + GPoolBlockSize;
 
 		// Link subsequent blocks together
@@ -164,10 +164,10 @@ static FWriteBuffer* Writer_NextBufferInternal(uint32 PageGrowth)
 		break;
 	}
 
-	GWriteBuffer = Next;
+	GWriteBuffer = NextBuffer;
 
-	Next->Cursor = ((uint8*)Next - GPoolBlockSize + sizeof(FWriteBuffer));
-	return Next;
+	NextBuffer->Cursor = ((uint8*)NextBuffer - GPoolBlockSize + sizeof(FWriteBuffer));
+	return NextBuffer;
 }
 
 ////////////////////////////////////////////////////////////////////////////////
