Use the correct FEventHeader type instead of uint32 when dealing with events.

#rnx

diff --git a/Engine/Source/Runtime/TraceLog/Private/Trace/Writer.cpp b/Engine/Source/Runtime/TraceLog/Private/Trace/Writer.cpp
--- a/Engine/Source/Runtime/TraceLog/Private/Trace/Writer.cpp
+++ b/Engine/Source/Runtime/TraceLog/Private/Trace/Writer.cpp
@@ -470,8 +470,8 @@ static void Writer_ConsumeEvents()
 			continue;
 		}
 
-		const uint16* Header = (uint16*)(UPTRINT(EventPtr) + sizeof(void*));
-		uint16 DataSize = Header[1] + sizeof(uint32);
+		const auto* Header = (FEventHeader*)(UPTRINT(EventPtr) + sizeof(void*));
+		uint16 DataSize = Header->Size + sizeof(*Header);
 
 		Collector.Write(Header, DataSize);
 
diff --git a/Engine/Source/Runtime/TraceLog/Public/Trace/Detail/Writer.inl b/Engine/Source/Runtime/TraceLog/Public/Trace/Detail/Writer.inl
--- a/Engine/Source/Runtime/TraceLog/Public/Trace/Detail/Writer.inl
+++ b/Engine/Source/Runtime/TraceLog/Public/Trace/Detail/Writer.inl
@@ -53,7 +53,7 @@ inline uint8* Writer_BeginLog(uint16 EventUid, uint16 Size)
 {
 	using namespace Private;
 
-	static const uint32 HeaderSize = sizeof(void*) + sizeof(uint32);
+	static const uint32 HeaderSize = sizeof(void*) + sizeof(FEventHeader);
 	uint32 AllocSize = ((Size + HeaderSize) + 7) & ~7;
 
 	FWriteBuffer* Buffer = Writer_GetBuffer();
@@ -74,7 +74,7 @@ inline void Writer_EndLog(uint8* EventData)
 {
 	using namespace Private;
 
-	EventData -= sizeof(void*) + sizeof(uint32);
+	EventData -= sizeof(void*) + sizeof(FEventHeader);
 
 	// Add the event into the master linked list of events.
 	while (true)
