Removed unused zero-sized array pragmas

#rnx

diff --git a/Engine/Source/Runtime/TraceLog/Public/Trace/Detail/Writer.inl b/Engine/Source/Runtime/TraceLog/Public/Trace/Detail/Writer.inl
--- a/Engine/Source/Runtime/TraceLog/Public/Trace/Detail/Writer.inl
+++ b/Engine/Source/Runtime/TraceLog/Public/Trace/Detail/Writer.inl
@@ -12,12 +12,6 @@ namespace Trace
 namespace Private
 {
 
-////////////////////////////////////////////////////////////////////////////////
-#if defined(_MSC_VER)
-	#pragma warning(push)
-	#pragma warning(disable : 4200) // non-standard zero-sized array
-#endif
-
 ////////////////////////////////////////////////////////////////////////////////
 struct FWriteBuffer
 {
@@ -29,10 +23,6 @@ struct FWriteBuffer
 	uint32				ThreadId;
 };
 
-#if defined(_MSC_VER)
-	#pragma warning(pop)
-#endif
-
 ////////////////////////////////////////////////////////////////////////////////
 extern TRACELOG_API void* volatile	GLastEvent;
 TRACELOG_API uint8*					Writer_NextBuffer(uint16);
