Added some basic events tracing some basic information about the internals of Trace.

#rnx

diff --git a/Engine/Source/Runtime/TraceLog/Private/Trace/Writer.cpp b/Engine/Source/Runtime/TraceLog/Private/Trace/Writer.cpp
--- a/Engine/Source/Runtime/TraceLog/Private/Trace/Writer.cpp
+++ b/Engine/Source/Runtime/TraceLog/Private/Trace/Writer.cpp
@@ -29,6 +29,21 @@ int32 Encode(const void*, int32, void*, int32);
 
 
 
+////////////////////////////////////////////////////////////////////////////////
+#define TRACE_PRIVATE_PERF 0
+#if TRACE_PRIVATE_PERF
+UE_TRACE_EVENT_BEGIN($Trace, WorkerThread, Always)
+	UE_TRACE_EVENT_FIELD(uint32, Cycles)
+	UE_TRACE_EVENT_FIELD(uint32, BytesSent)
+UE_TRACE_EVENT_END()
+
+UE_TRACE_EVENT_BEGIN($Trace, Memory, Always)
+	UE_TRACE_EVENT_FIELD(uint32, AllocSize)
+UE_TRACE_EVENT_END()
+#endif // TRACE_PRIVATE_PERF
+
+
+
 ////////////////////////////////////////////////////////////////////////////////
 inline void Writer_Yield()
 {
@@ -350,6 +365,9 @@ UPTRINT					GPendingDataHandle;	// = 0
 ////////////////////////////////////////////////////////////////////////////////
 static void Writer_ConsumeEvents()
 {
+	uint64 StartTsc = TimeGetTimestamp();
+	uint32 BytesSent = 0;
+
 	// Claim ownership of the latest chain of sent events.
 	void* LatestEvent;
 	for (;; Writer_Yield())
@@ -473,12 +491,22 @@ static void Writer_ConsumeEvents()
 		const auto* Header = (FEventHeader*)(UPTRINT(EventPtr) + sizeof(void*));
 		uint16 DataSize = Header->Size + sizeof(*Header);
 
+		BytesSent += DataSize;
 		Collector.Write(Header, DataSize);
 
 		EventPtr = *(void**)EventPtr;
 	}
 	Collector.Flush();
 
+#if TRACE_PRIVATE_PERF
+	UE_TRACE_LOG($Trace, WorkerThread)
+		<< WorkerThread.Cycles(uint32(TimeGetTimestamp() - StartTsc))
+		<< WorkerThread.BytesSent(BytesSent);
+
+	UE_TRACE_LOG($Trace, Memory)
+		<< Memory.AllocSize(uint32(GPoolPageCursor - GPoolBase));
+#endif // TRACE_PRIVATE_PERF
+
 	if (RetiredHead == nullptr)
 	{
 		return;
