Fixed missing casts to void*.

#rnx

diff --git a/Engine/Source/Runtime/TraceLog/Public/Trace/Detail/Atomic.h b/Engine/Source/Runtime/TraceLog/Public/Trace/Detail/Atomic.h
--- a/Engine/Source/Runtime/TraceLog/Public/Trace/Detail/Atomic.h
+++ b/Engine/Source/Runtime/TraceLog/Public/Trace/Detail/Atomic.h
@@ -141,7 +141,7 @@ inline bool AtomicCompareExchangeRelaxed(Type volatile* Target, Type New, Type E
 {
 	static_assert(sizeof(Type) == sizeof(void*), "");
 #if IS_MSVC
-	return INTERLOCKED_API(_InterlockedCompareExchangePointer, _nf, (void* volatile*)Target, New, Expected) == Expected;
+	return INTERLOCKED_API(_InterlockedCompareExchangePointer, _nf, (void* volatile*)Target, (void*)New, (void*)Expected) == (void*)Expected;
 #elif IS_GCC_COMPATIBLE
 	Type InOut = Expected;
 	return __atomic_compare_exchange_n(Target, &InOut, New, true, __ATOMIC_RELAXED, __ATOMIC_RELAXED);
@@ -154,7 +154,7 @@ inline bool AtomicCompareExchangeAcquire(Type volatile* Target, Type New, Type E
 {
 	static_assert(sizeof(Type) == sizeof(void*), "");
 #if IS_MSVC
-	return INTERLOCKED_API(_InterlockedCompareExchangePointer, _acq, (void* volatile*)Target, New, Expected) == Expected;
+	return INTERLOCKED_API(_InterlockedCompareExchangePointer, _acq, (void* volatile*)Target, (void*)New, (void*)Expected) == (void*)Expected;
 #elif IS_GCC_COMPATIBLE
 	Type InOut = Expected;
 	return __atomic_compare_exchange_n(Target, &InOut, New, true, __ATOMIC_ACQUIRE, __ATOMIC_RELAXED);
@@ -167,7 +167,7 @@ inline bool AtomicCompareExchangeRelease(Type volatile* Target, Type New, Type E
 {
 	static_assert(sizeof(Type) == sizeof(void*), "");
 #if IS_MSVC
-	return INTERLOCKED_API(_InterlockedCompareExchangePointer, _rel, (void* volatile*)Target, New, Expected) == Expected;
+	return INTERLOCKED_API(_InterlockedCompareExchangePointer, _rel, (void* volatile*)Target, (void*)New, (void*)Expected) == (void*)Expected;
 #elif IS_GCC_COMPATIBLE
 	Type InOut = Expected;
 	return __atomic_compare_exchange_n(Target, &InOut, New, true, __ATOMIC_RELEASE, __ATOMIC_RELAXED);
