Fixed missing and misordered includes.

#rnx

diff --git a/Engine/Source/Runtime/TraceLog/Public/Trace/Detail/Writer.inl b/Engine/Source/Runtime/TraceLog/Public/Trace/Detail/Writer.inl
--- a/Engine/Source/Runtime/TraceLog/Public/Trace/Detail/Writer.inl
+++ b/Engine/Source/Runtime/TraceLog/Public/Trace/Detail/Writer.inl
@@ -5,6 +5,7 @@
 #if UE_TRACE_ENABLED
 
 #include "Atomic.h"
+#include "Protocol.h"
 
 namespace Trace
 {
