Wrapped the current trace event buffer into a per-thread context object. Consolidates some of the writer's thread-related code to one place and gives a logical place where threads come up and down.

#rnx

diff --git a/Engine/Source/Runtime/TraceLog/Private/Trace/Writer.cpp b/Engine/Source/Runtime/TraceLog/Private/Trace/Writer.cpp
--- a/Engine/Source/Runtime/TraceLog/Private/Trace/Writer.cpp
+++ b/Engine/Source/Runtime/TraceLog/Private/Trace/Writer.cpp
@@ -78,14 +78,63 @@ void Writer_InitializeTiming()
 
 
 
+////////////////////////////////////////////////////////////////////////////////
+thread_local		FWriteTlsContext TlsContext;
+uint8				FWriteTlsContext::DefaultBuffer[];	// = {}
+UPTRINT volatile	FWriteTlsContext::ThreadIdCounter;	// = 0;
+
+////////////////////////////////////////////////////////////////////////////////
+FWriteTlsContext::FWriteTlsContext()
+{
+	auto* Target = (FWriteBuffer*)DefaultBuffer;
+
+	static bool Once;
+	if (!Once)
+	{
+		Target->Cursor = DefaultBuffer;
+		Once = true;
+	}
+
+	Buffer = Target;
+
+	// Assign a new id to this thread.
+	for (;; Writer_Yield())
+	{
+		UPTRINT CandidateId = UPTRINT(AtomicLoadRelaxed((void* volatile*)&ThreadIdCounter));
+		UPTRINT NextId = CandidateId + 1;
+		if (AtomicCompareExchangeRelaxed(&ThreadIdCounter, NextId, CandidateId))
+		{
+			ThreadId = uint32(CandidateId);
+			break;
+		}
+	}
+}
+
+////////////////////////////////////////////////////////////////////////////////
+inline bool FWriteTlsContext::HasValidBuffer() const
+{
+	return (UPTRINT(Buffer) != UPTRINT(DefaultBuffer));
+}
+
+////////////////////////////////////////////////////////////////////////////////
+inline void FWriteTlsContext::SetBuffer(FWriteBuffer* InBuffer)
+{
+	Buffer = InBuffer;
+}
+
+////////////////////////////////////////////////////////////////////////////////
+inline uint32 FWriteTlsContext::GetThreadId() const
+{
+	return ThreadId;
+}
+
+
+
 ////////////////////////////////////////////////////////////////////////////////
 #define T_ALIGN alignas(PLATFORM_CACHE_LINE_SIZE)
-static uint8						GEmptyBuffer[sizeof(FWriteBuffer)];
-thread_local FWriteBuffer*			GWriteBuffer		= (FWriteBuffer*)GEmptyBuffer;
 T_ALIGN static void* volatile		GFirstEvent;
 T_ALIGN TRACELOG_API void* volatile	GLastEvent;			// = nullptr;
 static const uint32					GPoolSize			= 384 << 20; // 384MB ought to be enough
-T_ALIGN static UPTRINT volatile		GThreadId;			// = 0;
 static const uint32					GPoolBlockSize		= 4 << 10;
 static const uint32					GPoolPageGrowth		= GPoolBlockSize << 5;
 static const uint32					GPoolInitPageSize	= GPoolBlockSize << 5;
@@ -107,7 +156,7 @@ TRACELOG_API FWriteBuffer* Writer_GetBuffer()
 {
 	// Thread locals and DLLs don't mix so for modular builds we are forced to
 	// export this function to access thread-local variables.
-	return GWriteBuffer;
+	return TlsContext.GetBuffer();
 }
 #endif
 
@@ -179,9 +228,9 @@ static FWriteBuffer* Writer_NextBufferInternal(uint32 PageGrowth)
 		break;
 	}
 
-	GWriteBuffer = NextBuffer;
-
 	NextBuffer->Cursor = ((uint8*)NextBuffer - GPoolBlockSize + sizeof(FWriteBuffer));
+
+	TlsContext.SetBuffer(NextBuffer);
 	return NextBuffer;
 }
 
@@ -194,30 +243,10 @@ TRACELOG_API uint8* Writer_NextBuffer(uint16 Size)
 		return nullptr;
 	}
 
-	FWriteBuffer* Current = GWriteBuffer;
-
-	// Carry along or assign a new thread id
-	uint32 ThreadId;
-	if (UPTRINT(Current) == UPTRINT(GEmptyBuffer))
-	{
-		for (;; Writer_Yield())
-		{
-			UPTRINT CurrentTid = UPTRINT(AtomicLoadRelaxed((void* volatile*)&GThreadId));
-			UPTRINT NextTid = CurrentTid + 1;
-			if (AtomicCompareExchangeRelaxed((void* volatile*)&GThreadId, (void*)NextTid, (void*)CurrentTid))
-			{
-				ThreadId = uint32(CurrentTid);
-				break;
-			}
-		}
-	}
-	else
-	{
-		ThreadId = Current->ThreadId;
-	}
+	FWriteBuffer* Current = TlsContext.GetBuffer();
 
 	// Retire current buffer unless its the initial boot one.
-	if (UPTRINT(Current) != UPTRINT(GEmptyBuffer))
+	if (TlsContext.HasValidBuffer())
 	{
 		// To retire a buffer we'll link it into the event list which event
 		// consumption can detect and use to return the buffer to the free list.
@@ -233,7 +262,7 @@ TRACELOG_API uint8* Writer_NextBuffer(uint16 Size)
 	}
 
 	FWriteBuffer* NextBuffer = Writer_NextBufferInternal(GPoolPageGrowth);
-	NextBuffer->ThreadId = ThreadId;
+	NextBuffer->ThreadId = TlsContext.GetThreadId();
 
 	NextBuffer->Cursor += Size;
 	return NextBuffer->Cursor;
@@ -249,9 +278,6 @@ static void Writer_InitializeBuffers()
 
 	static_assert(GPoolPageGrowth >= 0x10000, "Page growth must be >= 64KB");
 	static_assert(GPoolInitPageSize >= 0x10000, "Initial page size must be >= 64KB");
-
-	auto* EmptyBuffer = (FWriteBuffer*)GEmptyBuffer;
-	EmptyBuffer->Cursor = GEmptyBuffer;
 }
 
 ////////////////////////////////////////////////////////////////////////////////
diff --git a/Engine/Source/Runtime/TraceLog/Public/Trace/Detail/Writer.inl b/Engine/Source/Runtime/TraceLog/Public/Trace/Detail/Writer.inl
--- a/Engine/Source/Runtime/TraceLog/Public/Trace/Detail/Writer.inl
+++ b/Engine/Source/Runtime/TraceLog/Public/Trace/Detail/Writer.inl
@@ -24,21 +24,45 @@ struct FWriteBuffer
 	uint32				ThreadId;
 };
 
+
+
 ////////////////////////////////////////////////////////////////////////////////
-extern TRACELOG_API void* volatile	GLastEvent;
-TRACELOG_API uint8*					Writer_NextBuffer(uint16);
-TRACELOG_API FWriteBuffer*			Writer_GetBuffer();
+struct FWriteTlsContext
+{
+							FWriteTlsContext();
+	bool					HasValidBuffer() const;
+	void					SetBuffer(FWriteBuffer*);
+	uint32					GetThreadId() const;
+	FWriteBuffer*			GetBuffer() const { return Buffer; }
+
+private:
+	FWriteBuffer*			Buffer;
+	uint32					ThreadId;
+	static uint8			DefaultBuffer[sizeof(FWriteBuffer)];
+	static UPTRINT volatile	ThreadIdCounter;
+};
+
+
 
+////////////////////////////////////////////////////////////////////////////////
+extern TRACELOG_API void* volatile		GLastEvent;
+TRACELOG_API uint8*						Writer_NextBuffer(uint16);
+
+////////////////////////////////////////////////////////////////////////////////
 #if IS_MONOLITHIC
-extern thread_local FWriteBuffer* GWriteBuffer;
+extern thread_local FWriteTlsContext TlsContext;
 inline FWriteBuffer* Writer_GetBuffer()
 {
-	return GWriteBuffer;
+	return TlsContext.GetBuffer();
 }
+#else
+TRACELOG_API FWriteBuffer* Writer_GetBuffer();
 #endif
 
 } // Private
 
+
+
 ////////////////////////////////////////////////////////////////////////////////
 inline uint8* Writer_BeginLog(uint16 EventUid, uint16 Size)
 {
