Wrapped the data passed between begin/end() trace event calls in a struct for easier extension in the future.

#rnx

diff --git a/Engine/Source/Runtime/TraceLog/Private/Trace/Writer.cpp b/Engine/Source/Runtime/TraceLog/Private/Trace/Writer.cpp
--- a/Engine/Source/Runtime/TraceLog/Private/Trace/Writer.cpp
+++ b/Engine/Source/Runtime/TraceLog/Private/Trace/Writer.cpp
@@ -1118,7 +1118,9 @@ void Writer_EventCreate(
 	uint16 EventSize = sizeof(FNewEventEvent);
 	EventSize += sizeof(FNewEventEvent::Fields[0]) * FieldCount;
 	EventSize += NamesSize;
-	auto& Event = *(FNewEventEvent*)Writer_BeginLog(EventUid, EventSize);
+
+	FLogInstance LogInstance = Writer_BeginLog(EventUid, EventSize);
+	auto& Event = *(FNewEventEvent*)(LogInstance.Ptr);
 
 	// Write event's main properties.
 	Event.EventUid = uint16(Uid) & uint16(EKnownEventUids::UidMask);
@@ -1153,7 +1155,7 @@ void Writer_EventCreate(
 		WriteName(Desc.Name, Desc.NameSize);
 	}
 
-	Writer_EndLog(&(uint8&)Event);
+	Writer_EndLog(LogInstance);
 
 	// Add this new event into the list so we can look them up later.
 	for (;; Writer_Yield())
diff --git a/Engine/Source/Runtime/TraceLog/Public/Trace/Detail/EventDef.h b/Engine/Source/Runtime/TraceLog/Public/Trace/Detail/EventDef.h
--- a/Engine/Source/Runtime/TraceLog/Public/Trace/Detail/EventDef.h
+++ b/Engine/Source/Runtime/TraceLog/Public/Trace/Detail/EventDef.h
@@ -6,6 +6,8 @@
 
 #if UE_TRACE_ENABLED
 
+#include "Writer.inl"
+
 namespace Trace
 {
 
@@ -28,7 +30,7 @@ public:
 								FLogScope(uint16 EventUid, uint16 Size);
 								FLogScope(uint16 EventUid, uint16 Size, uint16 ExtraBytes);
 								~FLogScope();
-		uint8*					Ptr;
+		FLogInstance			Instance;
 	};
 
 	void*						Handle;
diff --git a/Engine/Source/Runtime/TraceLog/Public/Trace/Detail/EventDef.inl b/Engine/Source/Runtime/TraceLog/Public/Trace/Detail/EventDef.inl
--- a/Engine/Source/Runtime/TraceLog/Public/Trace/Detail/EventDef.inl
+++ b/Engine/Source/Runtime/TraceLog/Public/Trace/Detail/EventDef.inl
@@ -5,7 +5,6 @@
 #if UE_TRACE_ENABLED
 
 #include "EventDef.h"
-#include "Writer.inl"
 
 namespace Trace
 {
@@ -13,19 +12,19 @@ namespace Trace
 ////////////////////////////////////////////////////////////////////////////////
 inline FEventDef::FLogScope::FLogScope(uint16 EventUid, uint16 Size)
 {
-	Ptr = Writer_BeginLog(EventUid, Size);
+	Instance = Writer_BeginLog(EventUid, Size);
 }
 
 ////////////////////////////////////////////////////////////////////////////////
 inline FEventDef::FLogScope::FLogScope(uint16 EventUid, uint16 Size, uint16 ExtraBytes)
 {
-	Ptr = Writer_BeginLog(EventUid, Size + ExtraBytes);
+	Instance = Writer_BeginLog(EventUid, Size + ExtraBytes);
 }
 
 ////////////////////////////////////////////////////////////////////////////////
 inline FEventDef::FLogScope::~FLogScope()
 {
-	Writer_EndLog(Ptr);
+	Writer_EndLog(Instance);
 }
 
 
@@ -34,7 +33,7 @@ inline FEventDef::FLogScope::~FLogScope()
 template <typename ActionType>
 inline const FEventDef::FLogScope& operator << (const FEventDef::FLogScope& Lhs, const ActionType& Rhs)
 {
-	Rhs.Write(Lhs.Ptr);
+	Rhs.Write(Lhs.Instance.Ptr);
 	return Lhs;
 }
 
diff --git a/Engine/Source/Runtime/TraceLog/Public/Trace/Detail/Writer.inl b/Engine/Source/Runtime/TraceLog/Public/Trace/Detail/Writer.inl
--- a/Engine/Source/Runtime/TraceLog/Public/Trace/Detail/Writer.inl
+++ b/Engine/Source/Runtime/TraceLog/Public/Trace/Detail/Writer.inl
@@ -64,7 +64,13 @@ TRACELOG_API FWriteBuffer* Writer_GetBuffer();
 
 
 ////////////////////////////////////////////////////////////////////////////////
-inline uint8* Writer_BeginLog(uint16 EventUid, uint16 Size)
+struct FLogInstance
+{
+	uint8* Ptr;
+};
+
+////////////////////////////////////////////////////////////////////////////////
+inline FLogInstance Writer_BeginLog(uint16 EventUid, uint16 Size)
 {
 	using namespace Private;
 
@@ -81,14 +87,15 @@ inline uint8* Writer_BeginLog(uint16 EventUid, uint16 Size)
 
 	uint32* Out = (uint32*)(Cursor + sizeof(void*));
 	Out[0] = (uint32(Size) << 16)|uint32(EventUid);
-	return (uint8*)(Out + 1);
+	return {(uint8*)(Out + 1)};
 }
 
 ////////////////////////////////////////////////////////////////////////////////
-inline void Writer_EndLog(uint8* EventData)
+inline void Writer_EndLog(FLogInstance Instance)
 {
 	using namespace Private;
 
+	uint8* EventData = Instance.Ptr;
 	EventData -= sizeof(void*) + sizeof(FEventHeader);
 
 	// Add the event into the master linked list of events.
