NextBuffer() returns the buffer pointer rather than its cursor position.

#rnx

diff --git a/Engine/Source/Runtime/TraceLog/Private/Trace/Writer.cpp b/Engine/Source/Runtime/TraceLog/Private/Trace/Writer.cpp
--- a/Engine/Source/Runtime/TraceLog/Private/Trace/Writer.cpp
+++ b/Engine/Source/Runtime/TraceLog/Private/Trace/Writer.cpp
@@ -235,7 +235,7 @@ static FWriteBuffer* Writer_NextBufferInternal(uint32 PageGrowth)
 }
 
 ////////////////////////////////////////////////////////////////////////////////
-TRACELOG_API uint8* Writer_NextBuffer(uint16 Size)
+TRACELOG_API FWriteBuffer* Writer_NextBuffer(uint16 Size)
 {
 	if (Size >= GPoolBlockSize - sizeof(FWriteBuffer))
 	{
@@ -265,7 +265,7 @@ TRACELOG_API uint8* Writer_NextBuffer(uint16 Size)
 	NextBuffer->ThreadId = TlsContext.GetThreadId();
 
 	NextBuffer->Cursor += Size;
-	return NextBuffer->Cursor;
+	return NextBuffer;
 }
 
 ////////////////////////////////////////////////////////////////////////////////
diff --git a/Engine/Source/Runtime/TraceLog/Public/Trace/Detail/Writer.inl b/Engine/Source/Runtime/TraceLog/Public/Trace/Detail/Writer.inl
--- a/Engine/Source/Runtime/TraceLog/Public/Trace/Detail/Writer.inl
+++ b/Engine/Source/Runtime/TraceLog/Public/Trace/Detail/Writer.inl
@@ -46,7 +46,7 @@ private:
 
 ////////////////////////////////////////////////////////////////////////////////
 extern TRACELOG_API void* volatile		GLastEvent;
-TRACELOG_API uint8*						Writer_NextBuffer(uint16);
+TRACELOG_API FWriteBuffer*				Writer_NextBuffer(uint16);
 
 ////////////////////////////////////////////////////////////////////////////////
 #if IS_MONOLITHIC
@@ -81,7 +81,8 @@ inline FLogInstance Writer_BeginLog(uint16 EventUid, uint16 Size)
 	uint8* Cursor = (Buffer->Cursor += AllocSize);
 	if (UNLIKELY(UPTRINT(Cursor) > UPTRINT(Buffer)))
 	{
-		Cursor = Writer_NextBuffer(AllocSize);
+		Buffer = Writer_NextBuffer(AllocSize);
+		Cursor = Buffer->Cursor;
 	}
 	Cursor -= AllocSize;
 
