Removed uneccesary void* cast

#rnx

diff --git a/Engine/Source/Runtime/TraceLog/Private/Trace/Writer.cpp b/Engine/Source/Runtime/TraceLog/Private/Trace/Writer.cpp
--- a/Engine/Source/Runtime/TraceLog/Private/Trace/Writer.cpp
+++ b/Engine/Source/Runtime/TraceLog/Private/Trace/Writer.cpp
@@ -1074,7 +1074,7 @@ void Writer_EventCreate(
 	{
 		UPTRINT CurrentUid = UPTRINT(AtomicLoadRelaxed((void* volatile*)&GEventUidCounter));
 		UPTRINT NextUid = CurrentUid + 1;
-		if (AtomicCompareExchangeRelaxed((void* volatile*)&GEventUidCounter, (void*)NextUid, (void*)CurrentUid))
+		if (AtomicCompareExchangeRelaxed(&GEventUidCounter, NextUid, CurrentUid))
 		{
 			Uid = uint32(CurrentUid) + uint32(EKnownEventUids::User);
 			break;
