Collect traced events by polling buffers instead of building a linked list of events. This reduces contention between different threads tracing events and allows event data to be more densely packed and more variable in length.

#rnx

diff --git a/Engine/Source/Runtime/TraceLog/Private/Trace/Writer.cpp b/Engine/Source/Runtime/TraceLog/Private/Trace/Writer.cpp
--- a/Engine/Source/Runtime/TraceLog/Private/Trace/Writer.cpp
+++ b/Engine/Source/Runtime/TraceLog/Private/Trace/Writer.cpp
@@ -132,15 +132,14 @@ inline uint32 FWriteTlsContext::GetThreadId() const
 
 ////////////////////////////////////////////////////////////////////////////////
 #define T_ALIGN alignas(PLATFORM_CACHE_LINE_SIZE)
-T_ALIGN static void* volatile		GFirstEvent;
-T_ALIGN TRACELOG_API void* volatile	GLastEvent;			// = nullptr;
-static const uint32					GPoolSize			= 384 << 20; // 384MB ought to be enough
-static const uint32					GPoolBlockSize		= 4 << 10;
-static const uint32					GPoolPageGrowth		= GPoolBlockSize << 5;
-static const uint32					GPoolInitPageSize	= GPoolBlockSize << 5;
-static uint8*						GPoolBase;			// = nullptr;
-T_ALIGN static uint8* volatile		GPoolPageCursor;	// = nullptr;
-T_ALIGN static void* volatile		GPoolFreeList;		// = nullptr;
+static const uint32						GPoolSize			= 384 << 20; // 384MB ought to be enough
+static const uint32						GPoolBlockSize		= 4 << 10;
+static const uint32						GPoolPageGrowth		= GPoolBlockSize << 5;
+static const uint32						GPoolInitPageSize	= GPoolBlockSize << 5;
+static uint8*							GPoolBase;			// = nullptr;
+T_ALIGN static uint8* volatile			GPoolPageCursor;	// = nullptr;
+T_ALIGN static void* volatile			GPoolFreeList;		// = nullptr;
+T_ALIGN static FWriteBuffer* volatile	GNextBufferList;	// = nullptr;
 #undef T_ALIGN
 
 ////////////////////////////////////////////////////////////////////////////////
@@ -228,7 +227,20 @@ static FWriteBuffer* Writer_NextBufferInternal(uint32 PageGrowth)
 		break;
 	}
 
+	// Add this next buffer to the active list.
+	for (;; Writer_Yield())
+	{
+		NextBuffer->Next = AtomicLoadRelaxed(&GNextBufferList);
+		if (AtomicCompareExchangeRelease(&GNextBufferList, NextBuffer, NextBuffer->Next))
+		{
+			break;
+		}
+	}
+
 	NextBuffer->Cursor = ((uint8*)NextBuffer - GPoolBlockSize + sizeof(FWriteBuffer));
+	NextBuffer->Cursor += sizeof(uint32); // this is so we can preceed event data with a small header when sending.
+	NextBuffer->Committed = NextBuffer->Cursor;
+	NextBuffer->Reaped = NextBuffer->Cursor;
 
 	TlsContext.SetBuffer(NextBuffer);
 	return NextBuffer;
@@ -248,17 +260,8 @@ TRACELOG_API FWriteBuffer* Writer_NextBuffer(uint16 Size)
 	// Retire current buffer unless its the initial boot one.
 	if (TlsContext.HasValidBuffer())
 	{
-		// To retire a buffer we'll link it into the event list which event
-		// consumption can detect and use to return the buffer to the free list.
-		for (;; Writer_Yield())
-		{
-			void* Expected = AtomicLoadRelaxed(&GLastEvent);
-			*(void**)Current = Expected;
-			if (AtomicCompareExchangeRelease(&GLastEvent, (void*)Current, Expected))
-			{
-				break;
-			}
-		}
+		Current->Cursor -= Size;
+		AtomicStoreRelease<uint8* __restrict>(&Current->Committed, nullptr);
 	}
 
 	FWriteBuffer* NextBuffer = Writer_NextBufferInternal(GPoolPageGrowth);
@@ -387,6 +390,7 @@ static FHoldBuffer		GHoldBuffer;		// will init to zero.
 static UPTRINT			GDataHandle;		// = 0
 static EDataState		GDataState;			// = EDataState::Passive
 UPTRINT					GPendingDataHandle;	// = 0
+static FWriteBuffer*	GActiveBufferList;	// = nullptr;
 
 ////////////////////////////////////////////////////////////////////////////////
 static void Writer_ConsumeEvents()
@@ -395,134 +399,181 @@ static void Writer_ConsumeEvents()
 	uint32 BytesSent = 0;
 
 	// Claim ownership of the latest chain of sent events.
-	void* LatestEvent;
+	FWriteBuffer* __restrict NextBufferList;
 	for (;; Writer_Yield())
 	{
-		LatestEvent = AtomicLoadRelaxed(&GLastEvent);
-		if (AtomicCompareExchangeAcquire(&GLastEvent, (void*)nullptr, LatestEvent))
+		NextBufferList = AtomicLoadRelaxed(&GNextBufferList);
+		if (AtomicCompareExchangeAcquire(&GNextBufferList, (FWriteBuffer*)nullptr, NextBufferList))
 		{
 			break;
 		}
 	}
 
-	FWriteBuffer* RetiredHead = nullptr;
-	FWriteBuffer* RetiredTail = nullptr;
+	struct FRetireList
+	{
+		FWriteBuffer* __restrict Head = nullptr;
+		FWriteBuffer* __restrict Tail = nullptr;
+
+		void Insert(FWriteBuffer* __restrict Buffer)
+		{
+			Buffer->Next = Head;
+			Head = Buffer;
+			Tail = (Tail != nullptr) ? Tail : Head;
+		}
+	};
 
-	static const uint32 BufferSize = 8192 - 64; // "-64" is to allow overflow/encoding overhead
-	struct FCollector
+	auto Send = [&] (uint8* __restrict Data, uint32 Size)
 	{
-		void Send(const void* Data, uint32 Size)
+		BytesSent += Size;
+
+		struct FPacketBase
 		{
-			if (GDataState == EDataState::Sending)
-			{
-				// Transmit data to the io handle
-				if (GDataHandle)
-				{
-					if (!IoWrite(GDataHandle, Data, Size))
-					{
-						IoClose(GDataHandle);
-						GDataHandle = 0;
-					}
-				}
-			}
-			else
-			{
-				GHoldBuffer->Write(Data, Size);
+			uint16 Serial;
+			uint16 PacketSize;
+		};
+
+#if 0
+		struct FPacketEncoded
+			: public FPacketBase
+		{
+			uint16	DecodedSize;
+		};
+
+		struct FPacket
+			: public FPacketEncoded
+		{
+			uint8 Data[GPoolBlockSize + 64];
+		};
+
+		FPacket Packet;
+		Packet.Serial = 0x8000;
+		Packet.DecodedSize = uint16(Size);
+		Packet.PacketSize = Encode(Data, Packet.DecodedSize, Packet.Data, sizeof(Packet.Data));
+		Packet.PacketSize += sizeof(FPacketEncoded);
+
+		Data = (uint8*)&Packet;
+		Size = Packet.PacketSize;
+#else
+		static_assert(sizeof(FPacketBase) == sizeof(uint32), "");
+		Data -= sizeof(FPacketBase);
+		Size += sizeof(FPacketBase);
+		auto* Packet = (FPacketBase*)Data;
+		Packet->Serial = 0;
+		Packet->PacketSize = uint16(Size);
+#endif
 
-				// Did we overflow? Enter partial mode.
-				bool bOverflown = GHoldBuffer->IsFull();
-				if (bOverflown && GDataState != EDataState::Partial)
+		if (GDataState == EDataState::Sending)
+		{
+			// Transmit data to the io handle
+			if (GDataHandle)
+			{
+				if (!IoWrite(GDataHandle, Data, Size))
 				{
-					GDataState = EDataState::Partial;
+					IoClose(GDataHandle);
+					GDataHandle = 0;
 				}
 			}
 		}
-
-		void Flush()
+		else
 		{
-			if (Cursor == sizeof(Buffer))
+			GHoldBuffer->Write(Data, Size);
+
+			// Did we overflow? Enter partial mode.
+			bool bOverflown = GHoldBuffer->IsFull();
+			if (bOverflown && GDataState != EDataState::Partial)
 			{
-				return;
+				GDataState = EDataState::Partial;
 			}
+		}
+	};
 
-			struct FPacketBase
-			{
-				uint16	Serial;
-				uint16	PacketSize;
-			};
+	FRetireList RetireList;
 
-#if 0
-			struct FPacketEncoded
-				: public FPacketBase
-			{
-				uint16	DecodedSize;
-			};
-
-			struct FPacket
-				: public FPacketEncoded
-			{
-				uint8 Data[BufferSize + 64];
-			};
-
-			FPacket Packet;
-			Packet.Serial = Serial | 0x8000; // MSB indicates the packet is encoded
-			Packet.DecodedSize = uint16(sizeof(Buffer) - Cursor);
-			Packet.PacketSize = Encode(Buffer + Cursor, Packet.DecodedSize, Packet.Data, sizeof(Packet.Data));
-			Packet.PacketSize += sizeof(FPacketEncoded);
-			Send(&Packet, Packet.PacketSize);
-#else
-			Cursor -= sizeof(FPacketBase);
-			auto* Packet = (FPacketBase*)(Buffer + Cursor);
-			Packet->Serial = Serial;
-			Packet->PacketSize = uint16(sizeof(Buffer) - Cursor);
-			Send(Packet, Packet->PacketSize);
-#endif
+	// Next buffer list is newest first. Retire full ones and build a new list
+	// of buffers that are active (which gets reverse so oldest is first)
+	FWriteBuffer* __restrict NextActiveList = nullptr;
+	FWriteBuffer* __restrict NextRetireList = nullptr;
+	for (FWriteBuffer *__restrict Buffer = NextBufferList, *__restrict NextBuffer;
+		Buffer != nullptr;
+		Buffer = NextBuffer)
+	{
+		NextBuffer = Buffer->Next;
 
-			Cursor = sizeof(Buffer);
-			Serial++;
+		if (AtomicLoadAcquire(&Buffer->Committed) != nullptr)
+		{
+			Buffer->Next = NextActiveList;
+			NextActiveList = Buffer;
 		}
+		else
+		{
+			Buffer->Next = NextRetireList;
+			NextRetireList = Buffer;
+		}
+	}
 
-		void Write(const void* Data, uint32 Size)
+	// Send as much of the active list as we can. Buffers that are full are removed
+	// from the list. Note that the list's oldest-first order is maintained
+	FWriteBuffer* __restrict ActiveListHead = nullptr;
+	FWriteBuffer* __restrict ActiveListTail = nullptr;
+	for (FWriteBuffer *__restrict Buffer = GActiveBufferList, *__restrict NextBuffer;
+		Buffer != nullptr;
+		Buffer = NextBuffer)
+	{
+		NextBuffer = Buffer->Next;
+
+		uint8* __restrict Committed = AtomicLoadAcquire(&Buffer->Committed);
+		if (Committed == nullptr)
+		{
+			Committed = Buffer->Cursor;
+			RetireList.Insert(Buffer);
+		}
+		else
 		{
-			if (int32(Cursor - Size) < 0)
+			if (ActiveListTail != nullptr)
 			{
-				Flush();
+				ActiveListTail->Next = Buffer;
+			}
+			else
+			{
+				ActiveListHead = Buffer;
 			}
 
-			Cursor -= Size;
-			memcpy(Buffer + Cursor, Data, Size);
+			ActiveListTail = Buffer;
+			Buffer->Next = nullptr;
 		}
 
-		int16	Cursor = sizeof(Buffer);
-		uint16	Serial = 0;
-		uint8	Slack[8];				// Overflow for "Cursor - FPacketBase"
-		uint8	Buffer[BufferSize];
-	};
-
-	FCollector Collector;
-	for (void* EventPtr = LatestEvent; EventPtr != nullptr; )
-	{
-		// Is this "event" a retired buffer?
-		if ((UPTRINT(EventPtr) & (GPoolBlockSize - 1)) == (GPoolBlockSize - sizeof(FWriteBuffer)))
+		if (uint32 SizeToReap = uint32(Committed - Buffer->Reaped))
 		{
-			auto* Retiree = (FWriteBuffer*)EventPtr;
-			EventPtr = *(void**)EventPtr;
-
-			Retiree->Next = RetiredHead;
-			RetiredHead = Retiree;
-			RetiredTail = (RetiredTail != nullptr) ? RetiredTail : RetiredHead;
-			continue;
+			Send(Buffer->Reaped, SizeToReap);
+			Buffer->Reaped = Committed;
 		}
+	}
 
-		const auto* Header = (FEventHeader*)(UPTRINT(EventPtr) + sizeof(void*));
-		uint16 DataSize = Header->Size + sizeof(*Header);
+	// Retire buffers from the next list
+	for (FWriteBuffer *__restrict Buffer = NextRetireList, *__restrict NextBuffer;
+		Buffer != nullptr;
+		Buffer = NextBuffer)
+	{
+		NextBuffer = Buffer->Next;
 
-		BytesSent += DataSize;
-		Collector.Write(Header, DataSize);
+		RetireList.Insert(Buffer);
 
-		EventPtr = *(void**)EventPtr;
+		if (uint32 SizeToReap = uint32(Buffer->Cursor - Buffer->Reaped))
+		{
+			Send(Buffer->Reaped, SizeToReap);
+		}
+	}
+
+	// Append the new active buffers that have been discovered to the active list
+	if (ActiveListTail != nullptr)
+	{
+		GActiveBufferList = ActiveListHead;
+		ActiveListTail->Next = NextActiveList;
+	}
+	else
+	{
+		GActiveBufferList = NextActiveList;
 	}
-	Collector.Flush();
 
 #if TRACE_PRIVATE_PERF
 	UE_TRACE_LOG($Trace, WorkerThread)
@@ -533,18 +584,16 @@ static void Writer_ConsumeEvents()
 		<< Memory.AllocSize(uint32(GPoolPageCursor - GPoolBase));
 #endif // TRACE_PRIVATE_PERF
 
-	if (RetiredHead == nullptr)
-	{
-		return;
-	}
-
 	// Put the retirees we found back into the system again.
-	for (void** ListNode = (void**)RetiredTail;; Writer_Yield())
+	if (RetireList.Head != nullptr)
 	{
-		*ListNode = AtomicLoadRelaxed(&GPoolFreeList);
-		if (AtomicCompareExchangeRelease(&GPoolFreeList, (void*)RetiredHead, *ListNode))
+		for (void** ListNode = (void**)RetireList.Tail;; Writer_Yield())
 		{
-			break;
+			*ListNode = AtomicLoadRelaxed(&GPoolFreeList);
+			if (AtomicCompareExchangeRelease(&GPoolFreeList, (void*)RetireList.Head, *ListNode))
+			{
+				break;
+			}
 		}
 	}
 }
diff --git a/Engine/Source/Runtime/TraceLog/Public/Trace/Detail/Writer.inl b/Engine/Source/Runtime/TraceLog/Public/Trace/Detail/Writer.inl
--- a/Engine/Source/Runtime/TraceLog/Public/Trace/Detail/Writer.inl
+++ b/Engine/Source/Runtime/TraceLog/Public/Trace/Detail/Writer.inl
@@ -16,12 +16,11 @@ namespace Private
 ////////////////////////////////////////////////////////////////////////////////
 struct FWriteBuffer
 {
-	union
-	{
-		uint8*			Cursor;
-		FWriteBuffer*	Next;
-	};
-	uint32				ThreadId;
+	FWriteBuffer* __restrict	Next;
+	uint8* __restrict			Cursor;
+	uint8* __restrict volatile	Committed;
+	uint8* __restrict			Reaped;
+	uint32						ThreadId;
 };
 
 
@@ -66,7 +65,8 @@ TRACELOG_API FWriteBuffer* Writer_GetBuffer();
 ////////////////////////////////////////////////////////////////////////////////
 struct FLogInstance
 {
-	uint8* Ptr;
+	uint8*					Ptr;
+	Private::FWriteBuffer*	Internal;
 };
 
 ////////////////////////////////////////////////////////////////////////////////
@@ -74,8 +74,8 @@ inline FLogInstance Writer_BeginLog(uint16 EventUid, uint16 Size)
 {
 	using namespace Private;
 
-	static const uint32 HeaderSize = sizeof(void*) + sizeof(FEventHeader);
-	uint32 AllocSize = ((Size + HeaderSize) + 7) & ~7;
+	static const uint32 HeaderSize = sizeof(FEventHeader);
+	uint32 AllocSize = Size + HeaderSize;
 
 	FWriteBuffer* Buffer = Writer_GetBuffer();
 	uint8* Cursor = (Buffer->Cursor += AllocSize);
@@ -86,29 +86,17 @@ inline FLogInstance Writer_BeginLog(uint16 EventUid, uint16 Size)
 	}
 	Cursor -= AllocSize;
 
-	uint32* Out = (uint32*)(Cursor + sizeof(void*));
+	uint32* Out = (uint32*)Cursor;
 	Out[0] = (uint32(Size) << 16)|uint32(EventUid);
-	return {(uint8*)(Out + 1)};
+	return {(uint8*)(Out + 1), Buffer};
 }
 
 ////////////////////////////////////////////////////////////////////////////////
 inline void Writer_EndLog(FLogInstance Instance)
 {
 	using namespace Private;
-
-	uint8* EventData = Instance.Ptr;
-	EventData -= sizeof(void*) + sizeof(FEventHeader);
-
-	// Add the event into the master linked list of events.
-	while (true)
-	{
-		void* Expected = AtomicLoadRelaxed(&GLastEvent);
-		*(void**)EventData = Expected;
-		if (AtomicCompareExchangeRelease(&GLastEvent, (void*)EventData, Expected))
-		{
-			break;
-		}
-	}
+	FWriteBuffer* Buffer = Instance.Internal;
+	Private::AtomicStoreRelease<uint8* __restrict>(&(Buffer->Committed), Buffer->Cursor);
 }
 
 } // namespace Trace
