Added a destructor to FWriteTlsContext to return a thread's buffer to the free list.

#rnx

diff --git a/Engine/Source/Runtime/TraceLog/Private/Trace/Writer.cpp b/Engine/Source/Runtime/TraceLog/Private/Trace/Writer.cpp
--- a/Engine/Source/Runtime/TraceLog/Private/Trace/Writer.cpp
+++ b/Engine/Source/Runtime/TraceLog/Private/Trace/Writer.cpp
@@ -111,7 +111,16 @@ FWriteTlsContext::FWriteTlsContext()
 }
 
 ////////////////////////////////////////////////////////////////////////////////
-inline bool FWriteTlsContext::HasValidBuffer() const
+FWriteTlsContext::~FWriteTlsContext()
+{
+	if (HasValidBuffer())
+	{
+		AtomicStoreRelease<uint8* __restrict>(&(Buffer->Committed), nullptr);
+	}
+}
+
+////////////////////////////////////////////////////////////////////////////////
+bool FWriteTlsContext::HasValidBuffer() const
 {
 	return (UPTRINT(Buffer) != UPTRINT(DefaultBuffer));
 }
diff --git a/Engine/Source/Runtime/TraceLog/Public/Trace/Detail/Writer.inl b/Engine/Source/Runtime/TraceLog/Public/Trace/Detail/Writer.inl
--- a/Engine/Source/Runtime/TraceLog/Public/Trace/Detail/Writer.inl
+++ b/Engine/Source/Runtime/TraceLog/Public/Trace/Detail/Writer.inl
@@ -29,6 +29,7 @@ struct FWriteBuffer
 struct FWriteTlsContext
 {
 							FWriteTlsContext();
+							~FWriteTlsContext();
 	bool					HasValidBuffer() const;
 	void					SetBuffer(FWriteBuffer*);
 	uint32					GetThreadId() const;
