Removed unused extern.

#rnx

diff --git a/Engine/Source/Runtime/TraceLog/Public/Trace/Detail/Writer.inl b/Engine/Source/Runtime/TraceLog/Public/Trace/Detail/Writer.inl
--- a/Engine/Source/Runtime/TraceLog/Public/Trace/Detail/Writer.inl
+++ b/Engine/Source/Runtime/TraceLog/Public/Trace/Detail/Writer.inl
@@ -45,8 +45,8 @@ private:
 
 
 ////////////////////////////////////////////////////////////////////////////////
-extern TRACELOG_API void* volatile		GLastEvent;
-TRACELOG_API FWriteBuffer*				Writer_NextBuffer(uint16);
+TRACELOG_API FWriteBuffer* Writer_NextBuffer(uint16);
+TRACELOG_API FWriteBuffer* Writer_GetBuffer();
 
 ////////////////////////////////////////////////////////////////////////////////
 #if IS_MONOLITHIC
@@ -55,9 +55,7 @@ inline FWriteBuffer* Writer_GetBuffer()
 {
 	return TlsContext.GetBuffer();
 }
-#else
-TRACELOG_API FWriteBuffer* Writer_GetBuffer();
-#endif
+#endif // IS_MONOLITHIC
 
 } // Private
 
