Bumped the master trace version to 2 so Unreal Insights knows to look for a Diagnostics.Session event.

#rnx

diff --git a/Engine/Source/Runtime/TraceLog/Private/Trace/Writer.cpp b/Engine/Source/Runtime/TraceLog/Private/Trace/Writer.cpp
--- a/Engine/Source/Runtime/TraceLog/Private/Trace/Writer.cpp
+++ b/Engine/Source/Runtime/TraceLog/Private/Trace/Writer.cpp
@@ -979,7 +979,7 @@ static void Writer_LogHeader()
 	UE_TRACE_EVENT_END()
 
 	UE_TRACE_LOG($Trace, NewTrace)
-		<< NewTrace.Version(1)
+		<< NewTrace.Version(2)
 		<< NewTrace.Endian(0x524d)
 		<< NewTrace.PointerSize(sizeof(void*));
 }
