Removed explicit PLATFORM_PS4 from Trace/Config.h. This is taken care of by the TraceLog_PS4.Build.cs in Engine/Platforms/... now.

#rnx

diff --git a/Engine/Source/Runtime/TraceLog/Public/Trace/Config.h b/Engine/Source/Runtime/TraceLog/Public/Trace/Config.h
--- a/Engine/Source/Runtime/TraceLog/Public/Trace/Config.h
+++ b/Engine/Source/Runtime/TraceLog/Public/Trace/Config.h
@@ -6,7 +6,7 @@
 
 #if !defined(UE_TRACE_ENABLED)
 #	if !UE_BUILD_SHIPPING && !IS_PROGRAM
-#		if PLATFORM_WINDOWS || PLATFORM_UNIX || PLATFORM_APPLE || PLATFORM_PS4 || PLATFORM_XBOXONE || PLATFORM_SWITCH || PLATFORM_ANDROID
+#		if PLATFORM_WINDOWS || PLATFORM_UNIX || PLATFORM_APPLE || PLATFORM_XBOXONE || PLATFORM_SWITCH || PLATFORM_ANDROID
 #			define UE_TRACE_ENABLED	1
 #		endif
 #	endif
