Enabled LZ4-based compression on trace streams to reduce the file size and wire traffic.

#rnx

diff --git a/Engine/Source/Developer/TraceAnalysis/Private/Analysis/Transport/PacketTransport.cpp b/Engine/Source/Developer/TraceAnalysis/Private/Analysis/Transport/PacketTransport.cpp
--- a/Engine/Source/Developer/TraceAnalysis/Private/Analysis/Transport/PacketTransport.cpp
+++ b/Engine/Source/Developer/TraceAnalysis/Private/Analysis/Transport/PacketTransport.cpp
@@ -136,7 +136,6 @@ bool FPacketTransport::GetNextBatch()
 		Node->Next = PendingList;
 		PendingList = Node;
 
-#if 0
 		bool bEncoded = (PacketBase->Serial != PacketSerial);
 		if (bEncoded)
 		{
@@ -156,7 +155,6 @@ bool FPacketTransport::GetNextBatch()
 			);
 		}
 		else
-#endif // 0
 		{
 			struct FPacketRaw
 				: public FPacketBase
diff --git a/Engine/Source/Runtime/TraceLog/Private/Trace/Codec.cpp b/Engine/Source/Runtime/TraceLog/Private/Trace/Codec.cpp
new file mode 100644
--- /dev/null
+++ b/Engine/Source/Runtime/TraceLog/Private/Trace/Codec.cpp
@@ -0,0 +1,31 @@
+// Copyright 1998-2019 Epic Games, Inc. All Rights Reserved.
+
+#include "CoreTypes.h"
+
+#define LZ4_DISABLE_DEPRECATE_WARNINGS
+#include "Lz4/lz4.h"
+#undef LZ4_DISABLE_DEPRECATE_WARNINGS
+
+namespace Trace {
+namespace Private {
+
+////////////////////////////////////////////////////////////////////////////////
+int32 Encode(const void* Src, int32 SrcSize, void* Dest, int32 DestSize)
+{
+	return LZ4_compress_fast(
+		(const char*)Src,
+		(char*)Dest,
+		SrcSize,
+		DestSize,
+		1 // increase by 1 for small speed increase
+	);
+}
+
+////////////////////////////////////////////////////////////////////////////////
+TRACELOG_API int32 Decode(const void* Src, int32 SrcSize, void* Dest, int32 DestSize)
+{
+	return LZ4_decompress_safe((const char*)Src, (char*)Dest, SrcSize, DestSize);
+}
+
+} // namespace Private
+} // namespace Trace
diff --git a/Engine/Source/Runtime/TraceLog/Private/Trace/Writer.cpp b/Engine/Source/Runtime/TraceLog/Private/Trace/Writer.cpp
--- a/Engine/Source/Runtime/TraceLog/Private/Trace/Writer.cpp
+++ b/Engine/Source/Runtime/TraceLog/Private/Trace/Writer.cpp
@@ -441,35 +441,40 @@ static void Writer_ConsumeEvents()
 			uint16 PacketSize;
 		};
 
-#if 0
-		struct FPacketEncoded
-			: public FPacketBase
+		// Smaller buffers usually aren't redundant enough to benefit from being
+		// compressed. They often end up being larger.
+		if (Size > 384)
 		{
-			uint16	DecodedSize;
-		};
+			struct FPacketEncoded
+				: public FPacketBase
+			{
+				uint16	DecodedSize;
+			};
 
-		struct FPacket
-			: public FPacketEncoded
-		{
-			uint8 Data[GPoolBlockSize + 64];
-		};
+			struct FPacket
+				: public FPacketEncoded
+			{
+				uint8 Data[GPoolBlockSize + 64];
+			};
 
-		FPacket Packet;
-		Packet.Serial = 0x8000;
-		Packet.DecodedSize = uint16(Size);
-		Packet.PacketSize = Encode(Data, Packet.DecodedSize, Packet.Data, sizeof(Packet.Data));
-		Packet.PacketSize += sizeof(FPacketEncoded);
+			FPacket Packet;
+			Packet.Serial = 0x8000;
+			Packet.DecodedSize = uint16(Size);
+			Packet.PacketSize = Encode(Data, Packet.DecodedSize, Packet.Data, sizeof(Packet.Data));
+			Packet.PacketSize += sizeof(FPacketEncoded);
 
-		Data = (uint8*)&Packet;
-		Size = Packet.PacketSize;
-#else
-		static_assert(sizeof(FPacketBase) == sizeof(uint32), "");
-		Data -= sizeof(FPacketBase);
-		Size += sizeof(FPacketBase);
-		auto* Packet = (FPacketBase*)Data;
-		Packet->Serial = 0;
-		Packet->PacketSize = uint16(Size);
-#endif
+			Data = (uint8*)&Packet;
+			Size = Packet.PacketSize;
+		}
+		else
+		{
+			static_assert(sizeof(FPacketBase) == sizeof(uint32), "");
+			Data -= sizeof(FPacketBase);
+			Size += sizeof(FPacketBase);
+			auto* Packet = (FPacketBase*)Data;
+			Packet->Serial = 0;
+			Packet->PacketSize = uint16(Size);
+		}
 
 		if (GDataState == EDataState::Sending)
 		{
