Fixed Data being undefined on the compressed packet path.

#rnx

diff --git a/Engine/Source/Runtime/TraceLog/Private/Trace/Writer.cpp b/Engine/Source/Runtime/TraceLog/Private/Trace/Writer.cpp
--- a/Engine/Source/Runtime/TraceLog/Private/Trace/Writer.cpp
+++ b/Engine/Source/Runtime/TraceLog/Private/Trace/Writer.cpp
@@ -431,6 +431,33 @@ static void Writer_ConsumeEvents()
 		}
 	};
 
+	auto SendInner = [&] (uint8* __restrict Data, uint32 Size)
+	{
+		if (GDataState == EDataState::Sending)
+		{
+			// Transmit data to the io handle
+			if (GDataHandle)
+			{
+				if (!IoWrite(GDataHandle, Data, Size))
+				{
+					IoClose(GDataHandle);
+					GDataHandle = 0;
+				}
+			}
+		}
+		else
+		{
+			GHoldBuffer->Write(Data, Size);
+
+			// Did we overflow? Enter partial mode.
+			bool bOverflown = GHoldBuffer->IsFull();
+			if (bOverflown && GDataState != EDataState::Partial)
+			{
+				GDataState = EDataState::Partial;
+			}
+		}
+	};
+
 	auto Send = [&] (uint8* __restrict Data, uint32 Size)
 	{
 		BytesSent += Size;
@@ -463,8 +490,7 @@ static void Writer_ConsumeEvents()
 			Packet.PacketSize = Encode(Data, Packet.DecodedSize, Packet.Data, sizeof(Packet.Data));
 			Packet.PacketSize += sizeof(FPacketEncoded);
 
-			Data = (uint8*)&Packet;
-			Size = Packet.PacketSize;
+			SendInner((uint8*)&Packet, Packet.PacketSize);
 		}
 		else
 		{
@@ -474,30 +500,7 @@ static void Writer_ConsumeEvents()
 			auto* Packet = (FPacketBase*)Data;
 			Packet->Serial = 0;
 			Packet->PacketSize = uint16(Size);
-		}
-
-		if (GDataState == EDataState::Sending)
-		{
-			// Transmit data to the io handle
-			if (GDataHandle)
-			{
-				if (!IoWrite(GDataHandle, Data, Size))
-				{
-					IoClose(GDataHandle);
-					GDataHandle = 0;
-				}
-			}
-		}
-		else
-		{
-			GHoldBuffer->Write(Data, Size);
-
-			// Did we overflow? Enter partial mode.
-			bool bOverflown = GHoldBuffer->IsFull();
-			if (bOverflown && GDataState != EDataState::Partial)
-			{
-				GDataState = EDataState::Partial;
-			}
+			SendInner(Data, Size);
 		}
 	};
 
