Fixed Lz4 IWYU error.

#rnx

diff --git a/Engine/Source/Runtime/TraceLog/Private/Trace/Codec.cpp b/Engine/Source/Runtime/TraceLog/Private/Trace/Codec.cpp
--- a/Engine/Source/Runtime/TraceLog/Private/Trace/Codec.cpp
+++ b/Engine/Source/Runtime/TraceLog/Private/Trace/Codec.cpp
@@ -1,10 +1,7 @@
 // Copyright 1998-2019 Epic Games, Inc. All Rights Reserved.
 
 #include "CoreTypes.h"
-
-#define LZ4_DISABLE_DEPRECATE_WARNINGS
 #include "Lz4/lz4.h"
-#undef LZ4_DISABLE_DEPRECATE_WARNINGS
 
 namespace Trace {
 namespace Private {
diff --git a/Engine/Source/Runtime/TraceLog/TraceLog.Build.cs b/Engine/Source/Runtime/TraceLog/TraceLog.Build.cs
--- a/Engine/Source/Runtime/TraceLog/TraceLog.Build.cs
+++ b/Engine/Source/Runtime/TraceLog/TraceLog.Build.cs
@@ -10,5 +10,6 @@ public class TraceLog : ModuleRules
 	{
 		bRequiresImplementModule = false;
 		PublicIncludePathModuleNames.Add("Core");
+		PrivateDefinitions.Add("LZ4_DISABLE_DEPRECATE_WARNINGS");
     }
 }
