Appeased MSVC's static analysis of LZ4's code.
#rnx

diff --git a/Engine/Source/Runtime/TraceLog/Private/Trace/Codec.cpp b/Engine/Source/Runtime/TraceLog/Private/Trace/Codec.cpp
--- a/Engine/Source/Runtime/TraceLog/Private/Trace/Codec.cpp
+++ b/Engine/Source/Runtime/TraceLog/Private/Trace/Codec.cpp
@@ -2,10 +2,20 @@
 
 #include "CoreTypes.h"
 
+////////////////////////////////////////////////////////////////////////////////
+#if defined(_MSC_VER)
+#	pragma warning(push)
+#	pragma warning(disable : 6239)
+#endif
+
 extern "C" {
 #include "ThirdParty/LZ4/lz4.c"
 } // extern "C"
 
+#if defined(_MSC_VER)
+#	pragma warning(pop)
+#endif
+
 namespace Trace {
 namespace Private {
 
