Fixed MacOS and LZ4 symbol collision

#rnx

diff --git a/Engine/Source/Runtime/TraceLog/Private/Trace/Codec.cpp b/Engine/Source/Runtime/TraceLog/Private/Trace/Codec.cpp
--- a/Engine/Source/Runtime/TraceLog/Private/Trace/Codec.cpp
+++ b/Engine/Source/Runtime/TraceLog/Private/Trace/Codec.cpp
@@ -8,9 +8,10 @@
 #	pragma warning(disable : 6239)
 #endif
 
-extern "C" {
+namespace Lz4
+{
 #include "ThirdParty/LZ4/lz4.c"
-} // extern "C"
+}
 
 #if defined(_MSC_VER)
 #	pragma warning(pop)
@@ -22,7 +23,7 @@ namespace Private {
 ////////////////////////////////////////////////////////////////////////////////
 int32 Encode(const void* Src, int32 SrcSize, void* Dest, int32 DestSize)
 {
-	return LZ4_compress_fast(
+	return Lz4::LZ4_compress_fast(
 		(const char*)Src,
 		(char*)Dest,
 		SrcSize,
@@ -34,7 +35,7 @@ int32 Encode(const void* Src, int32 SrcSize, void* Dest, int32 DestSize)
 ////////////////////////////////////////////////////////////////////////////////
 TRACELOG_API int32 Decode(const void* Src, int32 SrcSize, void* Dest, int32 DestSize)
 {
-	return LZ4_decompress_safe((const char*)Src, (char*)Dest, SrcSize, DestSize);
+	return Lz4::LZ4_decompress_safe((const char*)Src, (char*)Dest, SrcSize, DestSize);
 }
 
 } // namespace Private
