Fixed the identifier collision on Mac another way.

#rnx

diff --git a/Engine/Source/Runtime/TraceLog/Private/Trace/Codec.cpp b/Engine/Source/Runtime/TraceLog/Private/Trace/Codec.cpp
--- a/Engine/Source/Runtime/TraceLog/Private/Trace/Codec.cpp
+++ b/Engine/Source/Runtime/TraceLog/Private/Trace/Codec.cpp
@@ -8,10 +8,7 @@
 #	pragma warning(disable : 6239)
 #endif
 
-namespace Lz4
-{
 #include "ThirdParty/LZ4/lz4.c"
-}
 
 #if defined(_MSC_VER)
 #	pragma warning(pop)
@@ -23,7 +20,7 @@ namespace Private {
 ////////////////////////////////////////////////////////////////////////////////
 int32 Encode(const void* Src, int32 SrcSize, void* Dest, int32 DestSize)
 {
-	return Lz4::LZ4_compress_fast(
+	return LZ4_compress_fast(
 		(const char*)Src,
 		(char*)Dest,
 		SrcSize,
@@ -35,7 +32,7 @@ int32 Encode(const void* Src, int32 SrcSize, void* Dest, int32 DestSize)
 ////////////////////////////////////////////////////////////////////////////////
 TRACELOG_API int32 Decode(const void* Src, int32 SrcSize, void* Dest, int32 DestSize)
 {
-	return Lz4::LZ4_decompress_safe((const char*)Src, (char*)Dest, SrcSize, DestSize);
+	return LZ4_decompress_safe((const char*)Src, (char*)Dest, SrcSize, DestSize);
 }
 
 } // namespace Private
