Added increment.

#rnx

diff --git a/Engine/Source/Runtime/TraceLog/Public/Trace/Detail/Atomic.h b/Engine/Source/Runtime/TraceLog/Public/Trace/Detail/Atomic.h
--- a/Engine/Source/Runtime/TraceLog/Public/Trace/Detail/Atomic.h
+++ b/Engine/Source/Runtime/TraceLog/Public/Trace/Detail/Atomic.h
@@ -13,6 +13,7 @@ template <typename Type> void	AtomicStoreRelease(Type volatile* Target, Type Val
 template <typename Type> bool	AtomicCompareExchangeRelaxed(Type volatile* Target, Type New, Type Expected);
 template <typename Type> bool	AtomicCompareExchangeAcquire(Type volatile* Target, Type New, Type Expected);
 template <typename Type> bool	AtomicCompareExchangeRelease(Type volatile* Target, Type New, Type Expected);
+uint32							AtomicIncrementRelaxed(uint32 volatile* Target);
 
 } // namespace Private
 } // namespace Trace
@@ -174,6 +175,18 @@ inline bool AtomicCompareExchangeRelease(Type volatile* Target, Type New, Type E
 #endif
 }
 
+////////////////////////////////////////////////////////////////////////////////
+inline uint32 AtomicIncrementRelaxed(uint32 volatile* Target)
+{
+	// Here we decrement the return of the MSVC path instead of incrementing the
+	// GCC path as GCC better matches what x64's 'lock add' does.
+#if IS_MSVC
+	return INTERLOCKED_API(_InterlockedIncrement, _nf, (long volatile*)Target) - 1;
+#elif IS_GCC_COMPATIBLE
+	return __atomic_add_fetch(Target, 1, __ATOMIC_RELAXED);
+#endif
+}
+
 #undef INTERLOCKED_API
 #undef IS_GCC_COMPATIBLE
 #undef IS_MSVC
