Changed field type enum into a scoped one.

#rnx

diff --git a/Engine/Source/Runtime/TraceLog/Public/Trace/Detail/Field.h b/Engine/Source/Runtime/TraceLog/Public/Trace/Detail/Field.h
--- a/Engine/Source/Runtime/TraceLog/Public/Trace/Detail/Field.h
+++ b/Engine/Source/Runtime/TraceLog/Public/Trace/Detail/Field.h
@@ -15,18 +15,18 @@ namespace Trace
 ////////////////////////////////////////////////////////////////////////////////
 template <typename Type> struct TFieldType;
 
-template <> struct TFieldType<bool>			{ enum { Value = Field_Bool }; };
-template <> struct TFieldType<int8>			{ enum { Value = Field_Int8 }; };
-template <> struct TFieldType<int16>		{ enum { Value = Field_Int16 }; };
-template <> struct TFieldType<int32>		{ enum { Value = Field_Int32 }; };
-template <> struct TFieldType<int64>		{ enum { Value = Field_Int64 }; };
-template <> struct TFieldType<uint8>		{ enum { Value = Field_Int8 }; };
-template <> struct TFieldType<uint16>		{ enum { Value = Field_Int16 }; };
-template <> struct TFieldType<uint32>		{ enum { Value = Field_Int32 }; };
-template <> struct TFieldType<uint64>		{ enum { Value = Field_Int64 }; };
-template <> struct TFieldType<float>		{ enum { Value = Field_Float32 }; };
-template <> struct TFieldType<double>		{ enum { Value = Field_Float64 }; };
-template <class T> struct TFieldType<T*>	{ enum { Value = Field_Pointer }; };
+template <> struct TFieldType<bool>			{ enum { Value = int(EFieldType::Bool) }; };
+template <> struct TFieldType<int8>			{ enum { Value = int(EFieldType::Int8) }; };
+template <> struct TFieldType<int16>		{ enum { Value = int(EFieldType::Int16) }; };
+template <> struct TFieldType<int32>		{ enum { Value = int(EFieldType::Int32) }; };
+template <> struct TFieldType<int64>		{ enum { Value = int(EFieldType::Int64) }; };
+template <> struct TFieldType<uint8>		{ enum { Value = int(EFieldType::Int8) }; };
+template <> struct TFieldType<uint16>		{ enum { Value = int(EFieldType::Int16) }; };
+template <> struct TFieldType<uint32>		{ enum { Value = int(EFieldType::Int32) }; };
+template <> struct TFieldType<uint64>		{ enum { Value = int(EFieldType::Int64) }; };
+template <> struct TFieldType<float>		{ enum { Value = int(EFieldType::Float32) }; };
+template <> struct TFieldType<double>		{ enum { Value = int(EFieldType::Float64) }; };
+template <class T> struct TFieldType<T*>	{ enum { Value = int(EFieldType::Pointer) }; };
 
 ////////////////////////////////////////////////////////////////////////////////
 struct FLiteralName
diff --git a/Engine/Source/Runtime/TraceLog/Public/Trace/Detail/Protocols/Protocol0.h b/Engine/Source/Runtime/TraceLog/Public/Trace/Detail/Protocols/Protocol0.h
--- a/Engine/Source/Runtime/TraceLog/Public/Trace/Detail/Protocols/Protocol0.h
+++ b/Engine/Source/Runtime/TraceLog/Public/Trace/Detail/Protocols/Protocol0.h
@@ -36,16 +36,16 @@ enum : uint8
 };
 
 ////////////////////////////////////////////////////////////////////////////////
-enum EFieldType : uint8
+enum class EFieldType : uint8
 {
-	Field_Bool		= Field_Integer | Field_8,
-	Field_Int8		= Field_Integer | Field_8,
-	Field_Int16		= Field_Integer | Field_16,
-	Field_Int32		= Field_Integer | Field_32,
-	Field_Int64		= Field_Integer | Field_64,
-	Field_Pointer	= Field_Integer | Field_Ptr,
-	Field_Float32	= Field_Float   | Field_32,
-	Field_Float64	= Field_Float   | Field_64,
+	Bool	= Field_Integer | Field_8,
+	Int8	= Field_Integer | Field_8,
+	Int16	= Field_Integer | Field_16,
+	Int32	= Field_Integer | Field_32,
+	Int64	= Field_Integer | Field_64,
+	Pointer	= Field_Integer | Field_Ptr,
+	Float32	= Field_Float   | Field_32,
+	Float64	= Field_Float   | Field_64,
 };
 
 ////////////////////////////////////////////////////////////////////////////////
