Use the GCC buildin add variant that returns the value prior to the add.

#rnx

diff --git a/Engine/Source/Runtime/TraceLog/Public/Trace/Detail/Atomic.h b/Engine/Source/Runtime/TraceLog/Public/Trace/Detail/Atomic.h
--- a/Engine/Source/Runtime/TraceLog/Public/Trace/Detail/Atomic.h
+++ b/Engine/Source/Runtime/TraceLog/Public/Trace/Detail/Atomic.h
@@ -183,7 +183,7 @@ inline uint32 AtomicIncrementRelaxed(uint32 volatile* Target)
 #if IS_MSVC
 	return INTERLOCKED_API(_InterlockedIncrement, _nf, (long volatile*)Target) - 1;
 #elif IS_GCC_COMPATIBLE
-	return __atomic_add_fetch(Target, 1, __ATOMIC_RELAXED);
+	return __atomic_fetch_add(Target, 1, __ATOMIC_RELAXED);
 #endif
 }
 
