Generalised TFieldBase so it's only purpose is to build Offset/Size values.

#rnx

diff --git a/Engine/Source/Runtime/TraceLog/Public/Trace/Detail/Field.h b/Engine/Source/Runtime/TraceLog/Public/Trace/Detail/Field.h
--- a/Engine/Source/Runtime/TraceLog/Public/Trace/Detail/Field.h
+++ b/Engine/Source/Runtime/TraceLog/Public/Trace/Detail/Field.h
@@ -63,17 +63,14 @@ struct FFieldDesc
 };
 
 ////////////////////////////////////////////////////////////////////////////////
-template <int OffsetValue, typename Type>
+template <int OffsetValue, int SizeValue>
 struct TFieldBase
-	: public FFieldDesc
 {
-	enum { Offset = OffsetValue, Size = sizeof(Type) };
-
-protected:
-	TFieldBase(const FLiteralName& Name)
-	: FFieldDesc(Name, TFieldType<Type>::Value, OffsetValue, Size)
+	enum : uint16
 	{
-	}
+		Offset	= OffsetValue,
+		Size	= SizeValue
+	};
 };
 
 ////////////////////////////////////////////////////////////////////////////////
@@ -148,8 +145,14 @@ struct TField<Offset, Type[Count]>
 ////////////////////////////////////////////////////////////////////////////////
 template <int Offset, typename Type>
 struct TField
-	: public TFieldBase<Offset, Type>
+	: public TFieldBase<Offset, sizeof(Type)>
+	, public FFieldDesc
 {
+	TField(const FLiteralName& Name)
+	: FFieldDesc(Name, TFieldType<Type>::Value, Offset, sizeof(Type))
+	{
+	}
+
 	struct FActionable
 	{
 		Type Value;
@@ -159,11 +162,6 @@ struct TField
 		}
 	};
 
-	TField(const FLiteralName& Name)
-	: TFieldBase<Offset, Type>(Name)
-	{
-	}
-
 	const FActionable operator () (const Type& __restrict Value) const
 	{
 		return { Value };
