Keep track of a fields' indices when declaring trace events

#rnx

diff --git a/Engine/Source/Runtime/TraceLog/Public/Trace/Detail/Field.h b/Engine/Source/Runtime/TraceLog/Public/Trace/Detail/Field.h
--- a/Engine/Source/Runtime/TraceLog/Public/Trace/Detail/Field.h
+++ b/Engine/Source/Runtime/TraceLog/Public/Trace/Detail/Field.h
@@ -63,33 +63,35 @@ struct FFieldDesc
 };
 
 ////////////////////////////////////////////////////////////////////////////////
-template <int OffsetValue, int SizeValue>
+template <int IndexValue, int OffsetValue, int SizeValue>
 struct TFieldBase
 {
 	enum : uint16
 	{
+		Index	= IndexValue,
 		Offset	= OffsetValue,
 		Size	= SizeValue
 	};
 };
 
 ////////////////////////////////////////////////////////////////////////////////
-template <int Offset, typename Type> struct TField;
+template <int Index, int Offset, typename Type> struct TField;
 
 ////////////////////////////////////////////////////////////////////////////////
 // Used to terminate the field list and determine an event's size.
 enum EndOfFields {};
-template <int Size>
-struct TField<Size, EndOfFields>
+template <int InFieldCount, int Size>
+struct TField<InFieldCount, Size, EndOfFields>
 {
-	enum : uint16 { Value = Size };
+	enum : uint16 { FieldCount = InFieldCount, Value = Size };
+	static_assert(FieldCount <= 127, "Trace events may only have up to a maximum of 127 fields");
 };
 
 ////////////////////////////////////////////////////////////////////////////////
 // Access to additional data that can be included along with a logged event.
 enum Attachment {};
 template <int Offset>
-struct TField<Offset, Attachment>
+struct TField<0, Offset, Attachment>
 {
 	template <typename LambdaType>
 	struct FActionableLambda
@@ -143,9 +145,9 @@ struct TField<Offset, Type[Count]>
 #endif
 
 ////////////////////////////////////////////////////////////////////////////////
-template <int Offset, typename Type>
+template <int IndexValue, int Offset, typename Type>
 struct TField
-	: public TFieldBase<Offset, sizeof(Type)>
+	: public TFieldBase<IndexValue, Offset, sizeof(Type)>
 	, public FFieldDesc
 {
 	TField(const FLiteralName& Name)
diff --git a/Engine/Source/Runtime/TraceLog/Public/Trace/Detail/Trace.h b/Engine/Source/Runtime/TraceLog/Public/Trace/Detail/Trace.h
--- a/Engine/Source/Runtime/TraceLog/Public/Trace/Detail/Trace.h
+++ b/Engine/Source/Runtime/TraceLog/Public/Trace/Detail/Trace.h
@@ -37,15 +37,17 @@
 				return true; \
 			}(); \
 		} \
-		Trace::TField<0, 
+		Trace::TField<0 /*Index*/, 0 /*Offset*/,
 
 #define TRACE_PRIVATE_EVENT_FIELD(FieldType, FieldName) \
 		FieldType> const FieldName = Trace::FLiteralName(#FieldName); \
-		Trace::TField<decltype(FieldName)::Offset + decltype(FieldName)::Size,
+		Trace::TField< \
+			decltype(FieldName)::Index + 1, \
+			decltype(FieldName)::Offset + decltype(FieldName)::Size,
 
 #define TRACE_PRIVATE_EVENT_END() \
 		Trace::EndOfFields> const EventSize_Private = {}; \
-		Trace::TField<decltype(EventSize_Private)::Value, Trace::Attachment> const Attachment = {}; \
+		Trace::TField<0, decltype(EventSize_Private)::Value, Trace::Attachment> const Attachment = {}; \
 		explicit operator bool () const { return true; } \
 	};
 
