Normalised names a little to maintain some consistency.

#rnx

diff --git a/Engine/Source/Runtime/TraceLog/Public/Trace/Detail/Field.h b/Engine/Source/Runtime/TraceLog/Public/Trace/Detail/Field.h
--- a/Engine/Source/Runtime/TraceLog/Public/Trace/Detail/Field.h
+++ b/Engine/Source/Runtime/TraceLog/Public/Trace/Detail/Field.h
@@ -63,14 +63,14 @@ struct FFieldDesc
 };
 
 ////////////////////////////////////////////////////////////////////////////////
-template <int IndexValue, int OffsetValue, int SizeValue>
+template <int InIndex, int InOffset, int InSize>
 struct TFieldBase
 {
 	enum : uint16
 	{
-		Index	= IndexValue,
-		Offset	= OffsetValue,
-		Size	= SizeValue
+		Index	= InIndex,
+		Offset	= InOffset,
+		Size	= InSize,
 	};
 };
 
@@ -145,9 +145,9 @@ struct TField<Offset, Type[Count]>
 #endif
 
 ////////////////////////////////////////////////////////////////////////////////
-template <int IndexValue, int Offset, typename Type>
+template <int Index, int Offset, typename Type>
 struct TField
-	: public TFieldBase<IndexValue, Offset, sizeof(Type)>
+	: public TFieldBase<Index, Offset, sizeof(Type)>
 	, public FFieldDesc
 {
 	TField(const FLiteralName& Name)
