Fixed unaligned access errors on ARM.

#rnx

diff --git a/Engine/Source/Runtime/TraceLog/Public/Trace/Detail/Field.h b/Engine/Source/Runtime/TraceLog/Public/Trace/Detail/Field.h
--- a/Engine/Source/Runtime/TraceLog/Public/Trace/Detail/Field.h
+++ b/Engine/Source/Runtime/TraceLog/Public/Trace/Detail/Field.h
@@ -160,7 +160,7 @@ struct TField
 		Type Value;
 		void Write(uint8* __restrict Ptr) const
 		{
-			*(Type*)(Ptr + Offset) = Value;
+			::memcpy(Ptr + Offset, &Value, sizeof(Type));
 		}
 	};
 
