The known event IDs enum is no longer misappropriated to also carry flags so the flag entries can be removed.

#rnx
#ushell-unshelve of 9957048

diff --git a/Engine/Source/Runtime/TraceLog/Public/Trace/Detail/Protocols/Protocol1.h b/Engine/Source/Runtime/TraceLog/Public/Trace/Detail/Protocols/Protocol1.h
--- a/Engine/Source/Runtime/TraceLog/Public/Trace/Detail/Protocols/Protocol1.h
+++ b/Engine/Source/Runtime/TraceLog/Public/Trace/Detail/Protocols/Protocol1.h
@@ -17,7 +17,6 @@ enum EProtocol : uint8 { Id = 3 };
 ////////////////////////////////////////////////////////////////////////////////
 using Protocol0::EFieldType;
 using Protocol0::FNewEventEvent;
-using Protocol0::EKnownEventUids;
 
 ////////////////////////////////////////////////////////////////////////////////
 enum class EEventFlags : uint8
@@ -25,6 +24,16 @@ enum class EEventFlags : uint8
 	Important		= 1 << 0,
 };
 
+////////////////////////////////////////////////////////////////////////////////
+enum class EKnownEventUids : uint16
+{
+	NewEvent,
+	User,
+	Max			= (1 << 15) - 1,
+	UidMask		= Max,
+	Invalid		= Max,
+};
+
 ////////////////////////////////////////////////////////////////////////////////
 struct FEventHeader
 {
