Moved size of a field into TFieldType to consolidate the sizeof() statements dotted about everywhere

#rnx
#ushell-unshelve of 9957048

diff --git a/Engine/Source/Runtime/TraceLog/Public/Trace/Detail/Field.h b/Engine/Source/Runtime/TraceLog/Public/Trace/Detail/Field.h
--- a/Engine/Source/Runtime/TraceLog/Public/Trace/Detail/Field.h
+++ b/Engine/Source/Runtime/TraceLog/Public/Trace/Detail/Field.h
@@ -15,18 +15,20 @@ namespace Trace
 ////////////////////////////////////////////////////////////////////////////////
 template <typename Type> struct TFieldType;
 
-template <> struct TFieldType<bool>			{ enum { Value = int(EFieldType::Bool) }; };
-template <> struct TFieldType<int8>			{ enum { Value = int(EFieldType::Int8) }; };
-template <> struct TFieldType<int16>		{ enum { Value = int(EFieldType::Int16) }; };
-template <> struct TFieldType<int32>		{ enum { Value = int(EFieldType::Int32) }; };
-template <> struct TFieldType<int64>		{ enum { Value = int(EFieldType::Int64) }; };
-template <> struct TFieldType<uint8>		{ enum { Value = int(EFieldType::Int8) }; };
-template <> struct TFieldType<uint16>		{ enum { Value = int(EFieldType::Int16) }; };
-template <> struct TFieldType<uint32>		{ enum { Value = int(EFieldType::Int32) }; };
-template <> struct TFieldType<uint64>		{ enum { Value = int(EFieldType::Int64) }; };
-template <> struct TFieldType<float>		{ enum { Value = int(EFieldType::Float32) }; };
-template <> struct TFieldType<double>		{ enum { Value = int(EFieldType::Float64) }; };
-template <class T> struct TFieldType<T*>	{ enum { Value = int(EFieldType::Pointer) }; };
+template <> struct TFieldType<bool>			{ enum { Tid = int(EFieldType::Bool),	Size = sizeof(bool) }; };
+template <> struct TFieldType<int8>			{ enum { Tid = int(EFieldType::Int8),	Size = sizeof(int8) }; };
+template <> struct TFieldType<int16>		{ enum { Tid = int(EFieldType::Int16),	Size = sizeof(int16) }; };
+template <> struct TFieldType<int32>		{ enum { Tid = int(EFieldType::Int32),	Size = sizeof(int32) }; };
+template <> struct TFieldType<int64>		{ enum { Tid = int(EFieldType::Int64),	Size = sizeof(int64) }; };
+template <> struct TFieldType<uint8>		{ enum { Tid = int(EFieldType::Int8),	Size = sizeof(uint8) }; };
+template <> struct TFieldType<uint16>		{ enum { Tid = int(EFieldType::Int16),	Size = sizeof(uint16) }; };
+template <> struct TFieldType<uint32>		{ enum { Tid = int(EFieldType::Int32),	Size = sizeof(uint32) }; };
+template <> struct TFieldType<uint64>		{ enum { Tid = int(EFieldType::Int64),	Size = sizeof(uint64) }; };
+template <> struct TFieldType<float>		{ enum { Tid = int(EFieldType::Float32),Size = sizeof(float) }; };
+template <> struct TFieldType<double>		{ enum { Tid = int(EFieldType::Float64),Size = sizeof(double) }; };
+template <class T> struct TFieldType<T*>	{ enum { Tid = int(EFieldType::Pointer),Size = sizeof(void*) }; };
+
+
 
 ////////////////////////////////////////////////////////////////////////////////
 struct FLiteralName
@@ -147,11 +149,11 @@ struct TField<Offset, Type[Count]>
 ////////////////////////////////////////////////////////////////////////////////
 template <int Index, int Offset, typename Type>
 struct TField
-	: public TFieldBase<Index, Offset, sizeof(Type)>
+	: public TFieldBase<Index, Offset, TFieldType<Type>::Size>
 	, public FFieldDesc
 {
 	TField(const FLiteralName& Name)
-	: FFieldDesc(Name, TFieldType<Type>::Value, Offset, sizeof(Type))
+	: FFieldDesc(Name, TFieldType<Type>::Tid, Offset, TFieldType<Type>::Size)
 	{
 	}
 
@@ -160,7 +162,7 @@ struct TField
 		Type Value;
 		void Write(uint8* __restrict Ptr) const
 		{
-			::memcpy(Ptr + Offset, &Value, sizeof(Type));
+			::memcpy(Ptr + Offset, &Value, TFieldType<Type>::Size);
 		}
 	};
 
