Encapsulated trace event field boilerplate in a macro.

#rnx
#ushell-unshelve of 9957048

diff --git a/Engine/Source/Runtime/TraceLog/Public/Trace/Detail/Field.h b/Engine/Source/Runtime/TraceLog/Public/Trace/Detail/Field.h
--- a/Engine/Source/Runtime/TraceLog/Public/Trace/Detail/Field.h
+++ b/Engine/Source/Runtime/TraceLog/Public/Trace/Detail/Field.h
@@ -65,16 +65,22 @@ struct FFieldDesc
 };
 
 ////////////////////////////////////////////////////////////////////////////////
-template <int InIndex, int InOffset, int InSize>
-struct TFieldBase
-{
-	enum : uint16
-	{
-		Index	= InIndex,
-		Offset	= InOffset,
-		Size	= InSize,
-	};
-};
+#define TRACE_PRIVATE_FIELD(InIndex, InOffset, Type) \
+		enum \
+		{ \
+			Index	= InIndex, \
+			Offset	= InOffset, \
+			Tid		= TFieldType<Type>::Tid, \
+			Size	= TFieldType<Type>::Size, \
+		}; \
+		static_assert(Index <= 127, "Trace events may only have up to a maximum of 127 fields"); \
+	private: \
+		FFieldDesc FieldDesc; \
+	public: \
+		TField(const FLiteralName& Name) \
+		: FieldDesc(Name, Tid, Offset, Size) \
+		{ \
+		}
 
 ////////////////////////////////////////////////////////////////////////////////
 template <int Index, int Offset, typename Type> struct TField;
@@ -147,22 +153,17 @@ struct TField<Offset, Type[Count]>
 #endif
 
 ////////////////////////////////////////////////////////////////////////////////
-template <int Index, int Offset, typename Type>
+template <int InIndex, int InOffset, typename Type>
 struct TField
-	: public TFieldBase<Index, Offset, TFieldType<Type>::Size>
-	, public FFieldDesc
 {
-	TField(const FLiteralName& Name)
-	: FFieldDesc(Name, TFieldType<Type>::Tid, Offset, TFieldType<Type>::Size)
-	{
-	}
+	TRACE_PRIVATE_FIELD(InIndex, InOffset, Type);
 
 	struct FActionable
 	{
 		Type Value;
 		void Write(uint8* __restrict Ptr) const
 		{
-			::memcpy(Ptr + Offset, &Value, TFieldType<Type>::Size);
+			::memcpy(Ptr + Offset, &Value, Size);
 		}
 	};
 
@@ -172,6 +173,8 @@ struct TField
 	}
 };
 
+#undef TRACE_PRIVATE_FIELD
+
 } // namespace Trace
 
 #endif // UE_TRACE_ENABLED
