Just a little reorganisation of the file.

#rnx

diff --git a/Engine/Source/Runtime/TraceLog/Public/Trace/Detail/Field.h b/Engine/Source/Runtime/TraceLog/Public/Trace/Detail/Field.h
--- a/Engine/Source/Runtime/TraceLog/Public/Trace/Detail/Field.h
+++ b/Engine/Source/Runtime/TraceLog/Public/Trace/Detail/Field.h
@@ -85,6 +85,51 @@
 ////////////////////////////////////////////////////////////////////////////////
 template <int Index, int Offset, typename Type> struct TField;
 
+
+
+#if 0
 ////////////////////////////////////////////////////////////////////////////////
+template <typename Type>
+struct TField<Type[]>
+	: public TFieldBase<sizeof(uint32)>
+{
+	TField(const ANSICHAR (Name)[]) : TFieldBase<sizeof(uint32)>(Name) {}
+};
+
+////////////////////////////////////////////////////////////////////////////////
+template <int Offset, typename Type, int Count>
+struct TField<Offset, Type[Count]>
+	: public TFieldBase<Offset, Type>
+{
+	TField(const ANSICHAR (Name)[]) : TFieldBase<Offset, Type[Count]>(Name) {}
+};
+#endif
+
+////////////////////////////////////////////////////////////////////////////////
+template <int InIndex, int InOffset, typename Type>
+struct TField
+{
+	TRACE_PRIVATE_FIELD(InIndex, InOffset, Type);
+
+	struct FActionable
+	{
+		Type Value;
+		void Write(uint8* __restrict Ptr) const
+		{
+			::memcpy(Ptr + Offset, &Value, Size);
+		}
+	};
+
+	const FActionable operator () (const Type& __restrict Value) const
+	{
+		return { Value };
+	}
+};
+
+#undef TRACE_PRIVATE_FIELD
+
+
+
+////////////////////////////////////////////////////////////////////////////////
 // Used to terminate the field list and determine an event's size.
 enum EndOfFields {};
@@ -134,47 +179,6 @@
 	}
 };
 
-#if 0
-////////////////////////////////////////////////////////////////////////////////
-template <typename Type>
-struct TField<Type[]>
-	: public TFieldBase<sizeof(uint32)>
-{
-	TField(const ANSICHAR (Name)[]) : TFieldBase<sizeof(uint32)>(Name) {}
-};
-
-////////////////////////////////////////////////////////////////////////////////
-template <int Offset, typename Type, int Count>
-struct TField<Offset, Type[Count]>
-	: public TFieldBase<Offset, Type>
-{
-	TField(const ANSICHAR (Name)[]) : TFieldBase<Offset, Type[Count]>(Name) {}
-};
-#endif
-
-////////////////////////////////////////////////////////////////////////////////
-template <int InIndex, int InOffset, typename Type>
-struct TField
-{
-	TRACE_PRIVATE_FIELD(InIndex, InOffset, Type);
-
-	struct FActionable
-	{
-		Type Value;
-		void Write(uint8* __restrict Ptr) const
-		{
-			::memcpy(Ptr + Offset, &Value, Size);
-		}
-	};
-
-	const FActionable operator () (const Type& __restrict Value) const
-	{
-		return { Value };
-	}
-};
-
-#undef TRACE_PRIVATE_FIELD
-
 } // namespace Trace
 
 #endif // UE_TRACE_ENABLED
