Removed unnecessary static_assert()

#rnx

diff --git a/Engine/Source/Runtime/TraceLog/Public/Trace/Detail/Field.h b/Engine/Source/Runtime/TraceLog/Public/Trace/Detail/Field.h
--- a/Engine/Source/Runtime/TraceLog/Public/Trace/Detail/Field.h
+++ b/Engine/Source/Runtime/TraceLog/Public/Trace/Detail/Field.h
@@ -137,7 +137,6 @@ template <int InFieldCount, int Size>
 struct TField<InFieldCount, Size, EndOfFields>
 {
 	enum : uint16 { FieldCount = InFieldCount, Value = Size };
-	static_assert(FieldCount <= 127, "Trace events may only have up to a maximum of 127 fields");
 };
 
 ////////////////////////////////////////////////////////////////////////////////
