Consistent naming convention between the various field implementations

#rnx
#ushell-unshelve of 9957048

diff --git a/Engine/Source/Runtime/TraceLog/Public/Trace/Detail/Field.h b/Engine/Source/Runtime/TraceLog/Public/Trace/Detail/Field.h
--- a/Engine/Source/Runtime/TraceLog/Public/Trace/Detail/Field.h
+++ b/Engine/Source/Runtime/TraceLog/Public/Trace/Detail/Field.h
@@ -142,8 +142,8 @@ struct TField<InFieldCount, Size, EndOfFields>
 ////////////////////////////////////////////////////////////////////////////////
 // Access to additional data that can be included along with a logged event.
 enum Attachment {};
-template <int Offset>
-struct TField<0, Offset, Attachment>
+template <int InOffset>
+struct TField<0, InOffset, Attachment>
 {
 	template <typename LambdaType>
 	struct FActionableLambda
@@ -151,7 +151,7 @@ struct TField<0, Offset, Attachment>
 		LambdaType& Value;
 		void Write(uint8* __restrict Ptr) const
 		{
-			Value(Ptr + Offset);
+			Value(Ptr + InOffset);
 		}
 	};
 
@@ -161,14 +161,13 @@ struct TField<0, Offset, Attachment>
 		return { Forward<LambdaType>(Lambda) };
 	}
 
-
 	struct FActionableMemcpy
 	{
 		const void* Data;
 		uint32 Size;
 		void Write(uint8* __restrict Ptr) const
 		{
-			::memcpy(Ptr + Offset, Data, Size);
+			::memcpy(Ptr + InOffset, Data, Size);
 		}
 	};
 
