Use the correct type for buffer lists instead of void**.

#rnx
#ushell-unshelve of 9957048

diff --git a/Engine/Source/Runtime/TraceLog/Private/Trace/Writer.cpp b/Engine/Source/Runtime/TraceLog/Private/Trace/Writer.cpp
--- a/Engine/Source/Runtime/TraceLog/Private/Trace/Writer.cpp
+++ b/Engine/Source/Runtime/TraceLog/Private/Trace/Writer.cpp
@@ -148,7 +148,7 @@ static const uint32						GPoolPageGrowth		= GPoolBlockSize << 5;
 static const uint32						GPoolInitPageSize	= GPoolBlockSize << 5;
 static uint8*							GPoolBase;			// = nullptr;
 T_ALIGN static uint8* volatile			GPoolPageCursor;	// = nullptr;
-T_ALIGN static void* volatile			GPoolFreeList;		// = nullptr;
+T_ALIGN static FWriteBuffer* volatile	GPoolFreeList;		// = nullptr;
 T_ALIGN static FWriteBuffer* volatile	GNextBufferList;	// = nullptr;
 #undef T_ALIGN
 
@@ -177,10 +177,10 @@ static FWriteBuffer* Writer_NextBufferInternal(uint32 PageGrowth)
 	while (true)
 	{
 		// First we'll try one from the free list
-		void* Owned = AtomicLoadRelaxed(&GPoolFreeList);
+		FWriteBuffer* Owned = AtomicLoadRelaxed(&GPoolFreeList);
 		if (Owned != nullptr)
 		{
-			if (!AtomicCompareExchangeRelaxed(&GPoolFreeList, *(void**)Owned, Owned))
+			if (!AtomicCompareExchangeRelaxed(&GPoolFreeList, Owned->Next, Owned))
 			{
 				Writer_Yield();
 				continue;
@@ -223,12 +223,11 @@ static FWriteBuffer* Writer_NextBufferInternal(uint32 PageGrowth)
 			Block += GPoolBlockSize;
 		}
 
-		// And insert the block list into the freelist
-		uint8* LastBlock = Block;
-		for (void** ListNode = (void**)LastBlock;; Writer_Yield())
+		// And insert the block list into the freelist. 'Block' is now the last block
+		for (auto* ListNode = (FWriteBuffer*)Block;; Writer_Yield())
 		{
-			*ListNode = AtomicLoadRelaxed(&GPoolFreeList);
-			if (AtomicCompareExchangeRelease(&GPoolFreeList, (void*)FirstBlock, *ListNode))
+			ListNode->Next = AtomicLoadRelaxed(&GPoolFreeList);
+			if (AtomicCompareExchangeRelease(&GPoolFreeList, (FWriteBuffer*)FirstBlock, ListNode->Next))
 			{
 				break;
 			}
@@ -606,10 +605,10 @@ static void Writer_ConsumeEvents()
 	// Put the retirees we found back into the system again.
 	if (RetireList.Head != nullptr)
 	{
-		for (void** ListNode = (void**)RetireList.Tail;; Writer_Yield())
+		for (FWriteBuffer* ListNode = RetireList.Tail;; Writer_Yield())
 		{
-			*ListNode = AtomicLoadRelaxed(&GPoolFreeList);
-			if (AtomicCompareExchangeRelease(&GPoolFreeList, (void*)RetireList.Head, *ListNode))
+			ListNode->Next = AtomicLoadRelaxed(&GPoolFreeList);
+			if (AtomicCompareExchangeRelease(&GPoolFreeList, RetireList.Head, ListNode->Next))
 			{
 				break;
 			}
