Just do a simple comparison of addresses. This pattern here is a relic from when a signed comparison was required.

#rnx

diff --git a/Engine/Source/Runtime/TraceLog/Public/Trace/Detail/Writer.inl b/Engine/Source/Runtime/TraceLog/Public/Trace/Detail/Writer.inl
--- a/Engine/Source/Runtime/TraceLog/Public/Trace/Detail/Writer.inl
+++ b/Engine/Source/Runtime/TraceLog/Public/Trace/Detail/Writer.inl
@@ -77,7 +77,7 @@ inline FLogInstance Writer_BeginLog(uint16 EventUid, uint16 Size, bool bMaybeHas
 	FWriteBuffer* Buffer = Writer_GetBuffer();
 	uint32 AllocSize = Size + sizeof(FEventHeader) + int(bMaybeHasAux);
 	Buffer->Cursor += AllocSize;
-	if (UNLIKELY(UPTRINT(Buffer->Cursor) > UPTRINT(Buffer)))
+	if (UNLIKELY(Buffer->Cursor > (uint8*)Buffer))
 	{
 		Buffer = Writer_NextBuffer(AllocSize);
 	}
