Added a little bit of overflow for a trace event buffer so we can write some of the protocol without having to bounds-check.

#rnx

diff --git a/Engine/Source/Runtime/TraceLog/Public/Trace/Detail/Writer.inl b/Engine/Source/Runtime/TraceLog/Public/Trace/Detail/Writer.inl
--- a/Engine/Source/Runtime/TraceLog/Public/Trace/Detail/Writer.inl
+++ b/Engine/Source/Runtime/TraceLog/Public/Trace/Detail/Writer.inl
@@ -16,11 +16,12 @@ namespace Private
 ////////////////////////////////////////////////////////////////////////////////
 struct FWriteBuffer
 {
+	uint32						Overflow;
+	uint32						ThreadId;
 	FWriteBuffer* __restrict	Next;
 	uint8* __restrict			Cursor;
 	uint8* __restrict volatile	Committed;
 	uint8* __restrict			Reaped;
-	uint32						ThreadId;
 };
 
 
