Indicate a buffer's exhausted by using an offset to where committed should be instead of setting committed to null. This allows committing written event data to be a separate step from marking a buffer as full.

#rnx
#ushell-unshelve of 9957048

diff --git a/Engine/Source/Runtime/TraceLog/Private/Trace/Writer.cpp b/Engine/Source/Runtime/TraceLog/Private/Trace/Writer.cpp
--- a/Engine/Source/Runtime/TraceLog/Private/Trace/Writer.cpp
+++ b/Engine/Source/Runtime/TraceLog/Private/Trace/Writer.cpp
@@ -116,7 +116,8 @@ FWriteTlsContext::~FWriteTlsContext()
 {
 	if (HasValidBuffer())
 	{
-		AtomicStoreRelease<uint8* __restrict>(&(Buffer->Committed), nullptr);
+		UPTRINT EtxOffset = ~UPTRINT((uint8*)Buffer - Buffer->Cursor);
+		AtomicStoreRelaxed(&(Buffer->EtxOffset), EtxOffset);
 	}
 }
 
@@ -240,6 +241,7 @@ static FWriteBuffer* Writer_NextBufferInternal(uint32 PageGrowth)
 	NextBuffer->Cursor += sizeof(uint32); // this is so we can preceed event data with a small header when sending.
 	NextBuffer->Committed = NextBuffer->Cursor;
 	NextBuffer->Reaped = NextBuffer->Cursor;
+	NextBuffer->EtxOffset = 0;
 
 	// Add this next buffer to the active list.
 	for (;; Writer_Yield())
@@ -269,8 +271,8 @@ TRACELOG_API FWriteBuffer* Writer_NextBuffer(uint16 Size)
 	// Retire current buffer unless its the initial boot one.
 	if (TlsContext.HasValidBuffer())
 	{
-		Current->Cursor -= Size;
-		AtomicStoreRelease<uint8* __restrict>(&Current->Committed, nullptr);
+		UPTRINT EtxOffset = ~UPTRINT((uint8*)(Current) - Current->Cursor + Size);
+		AtomicStoreRelease(&(Current->EtxOffset), EtxOffset);
 	}
 
 	FWriteBuffer* NextBuffer = Writer_NextBufferInternal(GPoolPageGrowth);
@@ -517,7 +519,9 @@ static void Writer_ConsumeEvents()
 	{
 		NextBuffer = Buffer->Next;
 
-		if (AtomicLoadAcquire(&Buffer->Committed) != nullptr)
+		uint8* Committed = AtomicLoadAcquire(&Buffer->Committed);
+		int32 EtxOffset = int32(~AtomicLoadRelaxed(&Buffer->EtxOffset));
+		if ((uint8*)Buffer - EtxOffset > Committed)
 		{
 			Buffer->Next = NextActiveList;
 			NextActiveList = Buffer;
@@ -540,9 +544,16 @@ static void Writer_ConsumeEvents()
 		NextBuffer = Buffer->Next;
 
 		uint8* __restrict Committed = AtomicLoadAcquire(&Buffer->Committed);
-		if (Committed == nullptr)
+
+		if (uint32 SizeToReap = uint32(Committed - Buffer->Reaped))
+		{
+			Send(Buffer, Buffer->Reaped, SizeToReap);
+			Buffer->Reaped = Committed;
+		}
+
+		int32 EtxOffset = int32(~AtomicLoadRelaxed(&Buffer->EtxOffset));
+		if ((uint8*)Buffer - EtxOffset == Committed)
 		{
-			Committed = Buffer->Cursor;
 			RetireList.Insert(Buffer);
 		}
 		else
@@ -559,12 +570,6 @@ static void Writer_ConsumeEvents()
 			ActiveListTail = Buffer;
 			Buffer->Next = nullptr;
 		}
-
-		if (uint32 SizeToReap = uint32(Committed - Buffer->Reaped))
-		{
-			Send(Buffer, Buffer->Reaped, SizeToReap);
-			Buffer->Reaped = Committed;
-		}
 	}
 
 	// Retire buffers from the next list
@@ -576,7 +581,7 @@ static void Writer_ConsumeEvents()
 
 		RetireList.Insert(Buffer);
 
-		if (uint32 SizeToReap = uint32(Buffer->Cursor - Buffer->Reaped))
+		if (uint32 SizeToReap = uint32(Buffer->Committed - Buffer->Reaped))
 		{
 			Send(Buffer, Buffer->Reaped, SizeToReap);
 		}
diff --git a/Engine/Source/Runtime/TraceLog/Public/Trace/Detail/Writer.inl b/Engine/Source/Runtime/TraceLog/Public/Trace/Detail/Writer.inl
--- a/Engine/Source/Runtime/TraceLog/Public/Trace/Detail/Writer.inl
+++ b/Engine/Source/Runtime/TraceLog/Public/Trace/Detail/Writer.inl
@@ -22,6 +22,7 @@ struct FWriteBuffer
 	uint8* __restrict			Cursor;
 	uint8* __restrict volatile	Committed;
 	uint8* __restrict			Reaped;
+	UPTRINT volatile			EtxOffset;
 };
 
 
