Removed Trace::GetMaxEventSize(). Not something I really want to expose to the public.

#rnx

diff --git a/Engine/Source/Runtime/TraceLog/Private/Trace/Trace.cpp b/Engine/Source/Runtime/TraceLog/Private/Trace/Trace.cpp
--- a/Engine/Source/Runtime/TraceLog/Private/Trace/Trace.cpp
+++ b/Engine/Source/Runtime/TraceLog/Private/Trace/Trace.cpp
@@ -13,7 +13,6 @@ namespace Private
 {
 
 ////////////////////////////////////////////////////////////////////////////////
-uint32	Writer_GetMaxEventSize();
 bool	Writer_SendTo(const ANSICHAR*);
 bool	Writer_WriteTo(const ANSICHAR*);
 uint32	Writer_EventToggle(const ANSICHAR*, bool);
@@ -36,12 +35,6 @@ static void ToAnsiCheap(ANSICHAR (&Dest)[DestSize], const WIDECHAR* Src)
 	}
 };
 
-////////////////////////////////////////////////////////////////////////////////
-uint32 GetMaxEventSize()
-{
-	return Private::Writer_GetMaxEventSize();
-}
-
 ////////////////////////////////////////////////////////////////////////////////
 bool SendTo(const TCHAR* InHost)
 {
diff --git a/Engine/Source/Runtime/TraceLog/Private/Trace/Writer.cpp b/Engine/Source/Runtime/TraceLog/Private/Trace/Writer.cpp
--- a/Engine/Source/Runtime/TraceLog/Private/Trace/Writer.cpp
+++ b/Engine/Source/Runtime/TraceLog/Private/Trace/Writer.cpp
@@ -153,13 +153,6 @@ T_ALIGN static FWriteBuffer* volatile	GPoolFreeList;		// = nullptr;
 T_ALIGN static FWriteBuffer* volatile	GNextBufferList;	// = nullptr;
 #undef T_ALIGN
 
-////////////////////////////////////////////////////////////////////////////////
-uint32 Writer_GetMaxEventSize()
-{
-	// The bais is to allow for some overhead. Its value was chosen arbitrarily.
-	return GPoolBlockSize - 512;
-}
-
 ////////////////////////////////////////////////////////////////////////////////
 #if !IS_MONOLITHIC
 TRACELOG_API FWriteBuffer* Writer_GetBuffer()
diff --git a/Engine/Source/Runtime/TraceLog/Public/Trace/Trace.h b/Engine/Source/Runtime/TraceLog/Public/Trace/Trace.h
--- a/Engine/Source/Runtime/TraceLog/Public/Trace/Trace.h
+++ b/Engine/Source/Runtime/TraceLog/Public/Trace/Trace.h
@@ -18,7 +18,6 @@
 namespace Trace
 {
 
-UE_TRACE_API uint32	GetMaxEventSize() UE_TRACE_IMPL(0);
 UE_TRACE_API bool	SendTo(const TCHAR* Host) UE_TRACE_IMPL(false);
 UE_TRACE_API bool	WriteTo(const TCHAR* Path) UE_TRACE_IMPL(false);
 UE_TRACE_API uint32 ToggleEvent(const TCHAR* Wildcard, bool bState) UE_TRACE_IMPL(0);
