Bump the protocol id value down one, back where it should be.

#rnx

diff --git a/Engine/Source/Runtime/TraceLog/Public/Trace/Detail/Protocols/Protocol1.h b/Engine/Source/Runtime/TraceLog/Public/Trace/Detail/Protocols/Protocol1.h
--- a/Engine/Source/Runtime/TraceLog/Public/Trace/Detail/Protocols/Protocol1.h
+++ b/Engine/Source/Runtime/TraceLog/Public/Trace/Detail/Protocols/Protocol1.h
@@ -12,7 +12,7 @@ namespace Protocol1
 {
 
 ////////////////////////////////////////////////////////////////////////////////
-enum EProtocol : uint8 { Id = 3 };
+enum EProtocol : uint8 { Id = 2 };
 
 ////////////////////////////////////////////////////////////////////////////////
 using Protocol0::EFieldType;
