Use correct thirdparty prologue and epilogue macros when including LZ4.

#rnx
#ushell-unshelve of 10483450

diff --git a/Engine/Source/Runtime/TraceLog/Private/Trace/Codec.cpp b/Engine/Source/Runtime/TraceLog/Private/Trace/Codec.cpp
--- a/Engine/Source/Runtime/TraceLog/Private/Trace/Codec.cpp
+++ b/Engine/Source/Runtime/TraceLog/Private/Trace/Codec.cpp
@@ -2,17 +2,9 @@
 
 #include "CoreTypes.h"
 
-////////////////////////////////////////////////////////////////////////////////
-#if defined(_MSC_VER)
-#	pragma warning(push)
-#	pragma warning(disable : 6239)
-#endif
-
+THIRD_PARTY_INCLUDES_START
 #include "ThirdParty/LZ4/lz4.c.h"
-
-#if defined(_MSC_VER)
-#	pragma warning(pop)
-#endif
+THIRD_PARTY_INCLUDES_END
 
 namespace Trace {
 namespace Private {
