Moved Trace's use of LZ4 into its namespace so there is no collision with other modules that may compile in LZ4

#rnx
#ushell-unshelve of 10483450

diff --git a/Engine/Source/Runtime/TraceLog/Private/Trace/Codec.cpp b/Engine/Source/Runtime/TraceLog/Private/Trace/Codec.cpp
--- a/Engine/Source/Runtime/TraceLog/Private/Trace/Codec.cpp
+++ b/Engine/Source/Runtime/TraceLog/Private/Trace/Codec.cpp
@@ -3,7 +3,9 @@
 #include "CoreTypes.h"
 
 THIRD_PARTY_INCLUDES_START
+#define LZ4_NAMESPACE Trace
 #include "ThirdParty/LZ4/lz4.c.h"
+#undef LZ4_NAMESPACE
 THIRD_PARTY_INCLUDES_END
 
 namespace Trace {
@@ -12,7 +14,7 @@ namespace Private {
 ////////////////////////////////////////////////////////////////////////////////
 int32 Encode(const void* Src, int32 SrcSize, void* Dest, int32 DestSize)
 {
-	return LZ4_compress_fast(
+	return Trace::LZ4_compress_fast(
 		(const char*)Src,
 		(char*)Dest,
 		SrcSize,
@@ -24,7 +26,7 @@ int32 Encode(const void* Src, int32 SrcSize, void* Dest, int32 DestSize)
 ////////////////////////////////////////////////////////////////////////////////
 TRACELOG_API int32 Decode(const void* Src, int32 SrcSize, void* Dest, int32 DestSize)
 {
-	return LZ4_decompress_safe((const char*)Src, (char*)Dest, SrcSize, DestSize);
+	return Trace::LZ4_decompress_safe((const char*)Src, (char*)Dest, SrcSize, DestSize);
 }
 
 } // namespace Private
