Ensure FLogScope object is created before any operator arguments. Depending on what the compiler decides to inline, things can either go completely fine or wildly wrong.

#rnx

diff --git a/Engine/Source/Runtime/TraceLog/Public/Trace/Detail/EventDef.h b/Engine/Source/Runtime/TraceLog/Public/Trace/Detail/EventDef.h
--- a/Engine/Source/Runtime/TraceLog/Public/Trace/Detail/EventDef.h
+++ b/Engine/Source/Runtime/TraceLog/Public/Trace/Detail/EventDef.h
@@ -32,6 +32,7 @@ public:
 								FLogScope(uint16 EventUid, uint16 Size, bool bMaybeHasAux, uint16 ExtraBytes);
 								~FLogScope();
 		FLogInstance			Instance;
+		constexpr explicit		operator bool () const { return true; }
 	};
 
 	void*						Handle;
diff --git a/Engine/Source/Runtime/TraceLog/Public/Trace/Detail/Trace.h b/Engine/Source/Runtime/TraceLog/Public/Trace/Detail/Trace.h
--- a/Engine/Source/Runtime/TraceLog/Public/Trace/Detail/Trace.h
+++ b/Engine/Source/Runtime/TraceLog/Public/Trace/Detail/Trace.h
@@ -64,11 +64,12 @@
 #define TRACE_PRIVATE_LOG(LoggerName, EventName, ...) \
 	if (TRACE_PRIVATE_EVENT_IS_ENABLED(LoggerName, EventName)) \
 		if (const auto& __restrict EventName = (F##LoggerName##EventName##Fields&)LoggerName##EventName##Event) \
-			Trace::FEventDef::FLogScope( \
+			if (auto LogScope = Trace::FEventDef::FLogScope( \
 				LoggerName##EventName##Event.Uid, \
 				TRACE_PRIVATE_EVENT_SIZE(LoggerName, EventName), \
 				bool(decltype(F##LoggerName##EventName##Fields::EventProps_Private)::MaybeHasAux), \
-				##__VA_ARGS__)
+				##__VA_ARGS__)) \
+					LogScope
 
 #else
 
