Implement the writing of array-type fields' data in terms of bytes instead of an element count

#rnx
#ushell-unshelve of 10495241

diff --git a/Engine/Source/Runtime/TraceLog/Public/Trace/Detail/Field.h b/Engine/Source/Runtime/TraceLog/Public/Trace/Detail/Field.h
--- a/Engine/Source/Runtime/TraceLog/Public/Trace/Detail/Field.h
+++ b/Engine/Source/Runtime/TraceLog/Public/Trace/Detail/Field.h
@@ -118,6 +118,13 @@ struct TField<InIndex, InOffset, Type[]>
 	};
 
 	const FActionable operator () (Type const* Data, int Count) const
+	{
+		Impl((const uint8*)Data, Count * sizeof(Type));
+		return {};
+	}
+
+private:
+	FORCENOINLINE void Impl(const uint8* Data, int32 Size) const
 	{
 		using namespace Private;
 
@@ -126,15 +133,15 @@ struct TField<InIndex, InOffset, Type[]>
 		FWriteBuffer* Buffer = Writer_GetBuffer();
 		Buffer->Cursor += sizeof(FAuxHeader) - bMaybeHasAux;
 
-		int32 Size = ((sizeof(Type) * Count) & 0x00ffffff);
+		Size &= 0x00ffffff;
 
 		auto* Header = (FAuxHeader*)(Buffer->Cursor - sizeof(FAuxHeader));
 		Header->Size = Size << 8;
 		Header->FieldIndex = uint8(0x80 | (Index & int(EIndexPack::FieldCountMask)));
 
+		bool bCommit = ((uint8*)Header == Buffer->Committed);
+
 		// Array data
-		const uint8* ReadCursor = (uint8*)Data;
-		bool bCommit = (UPTRINT(Header) == UPTRINT(Buffer->Committed));
 		while (true)
 		{
 			if (Buffer->Cursor >= (uint8*)Buffer)
@@ -145,7 +152,7 @@ struct TField<InIndex, InOffset, Type[]>
 
 			int32 Remaining = int32((uint8*)Buffer - Buffer->Cursor);
 			int32 SegmentSize = (Remaining < Size) ? Remaining : Size;
-			memcpy(Buffer->Cursor, ReadCursor, SegmentSize);
+			memcpy(Buffer->Cursor, Data, SegmentSize);
 			Buffer->Cursor += SegmentSize;
 
 			if (bCommit)
@@ -159,12 +166,11 @@ struct TField<InIndex, InOffset, Type[]>
 				break;
 			}
 
-			ReadCursor += SegmentSize;
+			Data += SegmentSize;
 		}
 
 		Buffer->Cursor[0] = 0;
 		Buffer->Cursor++;
-		return {};
 	}
 };
 
