Clamp array-type fields data to the maximum supported size (currently 16MB - should be plenty).=

#rnx
#ushell-unshelve of 10495241

diff --git a/Engine/Source/Runtime/TraceLog/Public/Trace/Detail/Field.h b/Engine/Source/Runtime/TraceLog/Public/Trace/Detail/Field.h
--- a/Engine/Source/Runtime/TraceLog/Public/Trace/Detail/Field.h
+++ b/Engine/Source/Runtime/TraceLog/Public/Trace/Detail/Field.h
@@ -119,7 +119,8 @@ struct TField<InIndex, InOffset, Type[]>
 
 	const FActionable operator () (Type const* Data, int Count) const
 	{
-		Impl((const uint8*)Data, Count * sizeof(Type));
+		int32 Size = (Count * sizeof(Type)) & (FAuxHeader::SizeLimit - 1) & ~(sizeof(Type) - 1);
+		Impl((const uint8*)Data, Size);
 		return {};
 	}
 
diff --git a/Engine/Source/Runtime/TraceLog/Public/Trace/Detail/Protocols/Protocol1.h b/Engine/Source/Runtime/TraceLog/Public/Trace/Detail/Protocols/Protocol1.h
--- a/Engine/Source/Runtime/TraceLog/Public/Trace/Detail/Protocols/Protocol1.h
+++ b/Engine/Source/Runtime/TraceLog/Public/Trace/Detail/Protocols/Protocol1.h
@@ -47,10 +47,11 @@ struct FEventHeader
 ////////////////////////////////////////////////////////////////////////////////
 struct FAuxHeader
 {
-	enum : uint8
+	enum : uint32
 	{
 		AuxDataBit	= 0x80,
 		FieldMask	= 0x7f,
+		SizeLimit	= 1 << 24,
 	};
 
 	union
