Forgot to remove old code that indiscriminately masked off the size of the data being written.

#rnx
#ushell-unshelve of 10495241

diff --git a/Engine/Source/Runtime/TraceLog/Public/Trace/Detail/Field.h b/Engine/Source/Runtime/TraceLog/Public/Trace/Detail/Field.h
--- a/Engine/Source/Runtime/TraceLog/Public/Trace/Detail/Field.h
+++ b/Engine/Source/Runtime/TraceLog/Public/Trace/Detail/Field.h
@@ -134,8 +134,6 @@ private:
 		FWriteBuffer* Buffer = Writer_GetBuffer();
 		Buffer->Cursor += sizeof(FAuxHeader) - bMaybeHasAux;
 
-		Size &= 0x00ffffff;
-
 		auto* Header = (FAuxHeader*)(Buffer->Cursor - sizeof(FAuxHeader));
 		Header->Size = Size << 8;
 		Header->FieldIndex = uint8(0x80 | (Index & int(EIndexPack::FieldCountMask)));
