int -> int32.

#rnx
#ushell-unshelve of 10495241

diff --git a/Engine/Source/Runtime/TraceLog/Public/Trace/Detail/Field.h b/Engine/Source/Runtime/TraceLog/Public/Trace/Detail/Field.h
--- a/Engine/Source/Runtime/TraceLog/Public/Trace/Detail/Field.h
+++ b/Engine/Source/Runtime/TraceLog/Public/Trace/Detail/Field.h
@@ -117,7 +117,7 @@ struct TField<InIndex, InOffset, Type[]>
 		void Write(uint8* __restrict) const {}
 	};
 
-	const FActionable operator () (Type const* Data, int Count) const
+	const FActionable operator () (Type const* Data, int32 Count) const
 	{
 		int32 Size = (Count * sizeof(Type)) & (FAuxHeader::SizeLimit - 1) & ~(sizeof(Type) - 1);
 		Impl((const uint8*)Data, Size);
