Do not write an array data if there will be none to write.

#rnx
#ushell-unshelve of 10495241

diff --git a/Engine/Source/Runtime/TraceLog/Public/Trace/Detail/Field.h b/Engine/Source/Runtime/TraceLog/Public/Trace/Detail/Field.h
--- a/Engine/Source/Runtime/TraceLog/Public/Trace/Detail/Field.h
+++ b/Engine/Source/Runtime/TraceLog/Public/Trace/Detail/Field.h
@@ -119,8 +119,11 @@ struct TField<InIndex, InOffset, Type[]>
 
 	const FActionable operator () (Type const* Data, int32 Count) const
 	{
-		int32 Size = (Count * sizeof(Type)) & (FAuxHeader::SizeLimit - 1) & ~(sizeof(Type) - 1);
-		Impl((const uint8*)Data, Size);
+		if (Count > 0)
+		{
+			int32 Size = (Count * sizeof(Type)) & (FAuxHeader::SizeLimit - 1) & ~(sizeof(Type) - 1);
+			Impl((const uint8*)Data, Size);
+		}
 		return {};
 	}
 
