Added a couple of comments to indicate where the end-of-aux data marker is written.

#rnx

diff --git a/Engine/Source/Runtime/TraceLog/Public/Trace/Detail/Field.h b/Engine/Source/Runtime/TraceLog/Public/Trace/Detail/Field.h
--- a/Engine/Source/Runtime/TraceLog/Public/Trace/Detail/Field.h
+++ b/Engine/Source/Runtime/TraceLog/Public/Trace/Detail/Field.h
@@ -171,6 +171,7 @@ private:
 			Data += SegmentSize;
 		}
 
+		// The auxilary data null terminator.
 		Buffer->Cursor[0] = 0;
 		Buffer->Cursor++;
 	}
diff --git a/Engine/Source/Runtime/TraceLog/Public/Trace/Detail/Writer.inl b/Engine/Source/Runtime/TraceLog/Public/Trace/Detail/Writer.inl
--- a/Engine/Source/Runtime/TraceLog/Public/Trace/Detail/Writer.inl
+++ b/Engine/Source/Runtime/TraceLog/Public/Trace/Detail/Writer.inl
@@ -84,6 +84,7 @@ inline FLogInstance Writer_BeginLog(uint16 EventUid, uint16 Size, bool bMaybeHas
 		Buffer = Writer_NextBuffer(AllocSize);
 	}
 
+	// The auxilary data null terminator.
 	if (bMaybeHasAux)
 	{
 		Buffer->Cursor[-1] = 0;
