Sketched out the scaffolding for fixed-size array type trace event fields.

#rnx

diff --git a/Engine/Source/Runtime/TraceLog/Public/Trace/Detail/Field.h b/Engine/Source/Runtime/TraceLog/Public/Trace/Detail/Field.h
--- a/Engine/Source/Runtime/TraceLog/Public/Trace/Detail/Field.h
+++ b/Engine/Source/Runtime/TraceLog/Public/Trace/Detail/Field.h
@@ -39,6 +39,18 @@ struct TFieldType<T[]>
 	};
 };
 
+#if 0
+template <typename T, int N>
+struct TFieldType<T[N]>
+{
+	enum
+	{
+		Tid  = int(TFieldType<T>::Tid)|int(EFieldType::Array),
+		Size = sizeof(T[N]),
+	};
+};
+#endif // 0
+
 
 
 ////////////////////////////////////////////////////////////////////////////////
@@ -177,6 +189,24 @@ private:
 	}
 };
 
+#if 0
+////////////////////////////////////////////////////////////////////////////////
+template <int InIndex, int InOffset, typename Type, int Count>
+struct TField<InIndex, InOffset, Type[Count]>
+{
+	TRACE_PRIVATE_FIELD(InIndex, InOffset, Type[Count]);
+
+	struct FActionable
+	{
+		void Write(uint8* __restrict) const {}
+	};
+
+	const FActionable operator () (Type const* Data, int Count) const
+	{
+	}
+};
+#endif
+
 ////////////////////////////////////////////////////////////////////////////////
 template <int InIndex, int InOffset, typename Type>
 struct TField
