Silence static analysis warnings in third party source code.

#rnx

diff --git a/Engine/Source/Runtime/TraceLog/Private/Trace/Codec.cpp b/Engine/Source/Runtime/TraceLog/Private/Trace/Codec.cpp
--- a/Engine/Source/Runtime/TraceLog/Private/Trace/Codec.cpp
+++ b/Engine/Source/Runtime/TraceLog/Private/Trace/Codec.cpp
@@ -3,9 +3,18 @@
 #include "CoreTypes.h"
 
 THIRD_PARTY_INCLUDES_START
+#if defined(_MSC_VER)
+#	pragma warning(push)
+#	pragma warning(disable : 6239)
+#endif
+
 #define LZ4_NAMESPACE Trace
 #include "LZ4/lz4.c.inl"
 #undef LZ4_NAMESPACE
+
+#if defined(_MSC_VER)
+#	pragma warning(pop)
+#endif
 THIRD_PARTY_INCLUDES_END
 
 namespace Trace {
