Added header guard to calm IncludeTool down.

#rnx

diff --git a/Engine/Source/Runtime/TraceLog/Private/Trace/LZ4/lz4.h b/Engine/Source/Runtime/TraceLog/Private/Trace/LZ4/lz4.h
--- a/Engine/Source/Runtime/TraceLog/Private/Trace/LZ4/lz4.h
+++ b/Engine/Source/Runtime/TraceLog/Private/Trace/LZ4/lz4.h
@@ -32,6 +32,9 @@
     - LZ4 homepage : http://www.lz4.org
     - LZ4 source repository : [REDACTED_PROFILE_URL]/lz4
 */
+
+#pragma once // EPIC MOD: Placate IncludeTool warning
+
 #if defined (__cplusplus)
 extern "C" {
 #endif
