Fixed PVS analysis warning.

#rnx

diff --git a/Engine/Source/Runtime/TraceLog/Private/Trace/LZ4/lz4.c.inl b/Engine/Source/Runtime/TraceLog/Private/Trace/LZ4/lz4.c.inl
--- a/Engine/Source/Runtime/TraceLog/Private/Trace/LZ4/lz4.c.inl
+++ b/Engine/Source/Runtime/TraceLog/Private/Trace/LZ4/lz4.c.inl
@@ -165,10 +165,10 @@
 #endif
 
 #ifndef likely
-#define likely(expr)     expect((expr) != 0, 1)
+#define likely(expr)     expect(int(expr) != 0, 1)
 #endif
 #ifndef unlikely
-#define unlikely(expr)   expect((expr) != 0, 0)
+#define unlikely(expr)   expect(int(expr) != 0, 0)
 #endif
 
 
