Fixed divergence of protocol version numbers between Main and Dev-Core.

#rnx

diff --git a/Engine/Source/Developer/TraceAnalysis/Private/Analysis/Engine.cpp b/Engine/Source/Developer/TraceAnalysis/Private/Analysis/Engine.cpp
--- a/Engine/Source/Developer/TraceAnalysis/Private/Analysis/Engine.cpp
+++ b/Engine/Source/Developer/TraceAnalysis/Private/Analysis/Engine.cpp
@@ -968,7 +968,7 @@ int32 FAnalysisEngine::OnDataProtocol1(FStreamReader& Reader)
 			break;
 		}
 
-		uint16 Uid = Header->Uid;
+		uint16 Uid = Header->Uid & uint16(Protocol1::EKnownEventUids::UidMask);
 		if (Uid >= Dispatches.Num())
 		{
 			return -1;
diff --git a/Engine/Source/Runtime/TraceLog/Public/Trace/Detail/Protocols/Protocol1.h b/Engine/Source/Runtime/TraceLog/Public/Trace/Detail/Protocols/Protocol1.h
--- a/Engine/Source/Runtime/TraceLog/Public/Trace/Detail/Protocols/Protocol1.h
+++ b/Engine/Source/Runtime/TraceLog/Public/Trace/Detail/Protocols/Protocol1.h
@@ -12,7 +12,7 @@ namespace Protocol1
 {
 
 ////////////////////////////////////////////////////////////////////////////////
-enum EProtocol : uint8 { Id = 3 };
+enum EProtocol : uint8 { Id = 1 };
 
 ////////////////////////////////////////////////////////////////////////////////
 using Protocol0::EFieldType;
