Removed stale #if/endif comment.

#rnx

diff --git a/Engine/Source/Runtime/TraceLog/Public/Trace/Config.h b/Engine/Source/Runtime/TraceLog/Public/Trace/Config.h
--- a/Engine/Source/Runtime/TraceLog/Public/Trace/Config.h
+++ b/Engine/Source/Runtime/TraceLog/Public/Trace/Config.h
@@ -10,7 +10,7 @@
 #			define UE_TRACE_ENABLED	1
 #		endif
 #	endif
-#endif // !IS_PROGRAM
+#endif
 
 #if !defined(UE_TRACE_ENABLED)
 #	define UE_TRACE_ENABLED 0
