Mac ununity compile fix.

#rnx

diff --git a/Engine/Source/Runtime/TraceLog/Public/Trace/Detail/Field.h b/Engine/Source/Runtime/TraceLog/Public/Trace/Detail/Field.h
--- a/Engine/Source/Runtime/TraceLog/Public/Trace/Detail/Field.h
+++ b/Engine/Source/Runtime/TraceLog/Public/Trace/Detail/Field.h
@@ -6,6 +6,7 @@
 
 #if UE_TRACE_ENABLED
 
+#include "Atomic.h"
 #include "Protocol.h"
 #include "Templates/UnrealTemplate.h"
 #include "Writer.inl"
@@ -171,7 +172,7 @@ private:
 
 			if (bCommit)
 			{
-				AtomicStoreRelease<uint8* __restrict>(&(Buffer->Committed), Buffer->Cursor);
+				AtomicStoreRelease(&(uint8* volatile&)(Buffer->Committed), Buffer->Cursor);
 			}
 
 			Size -= SegmentSize;
