Trace channels

This change introduces the concept of "channels" to TraceLog. When an event is emitted a channel now needs to be supplied. Channels are disabled by default and can be be enabled individually, allowing the user to limit the amount of events the tracing application produces. Channels can be enabled by the recording application (e.g. UnrealInsights) or enabled by default using a commandline argument, "-tracechannels" to enable a comma separated list of channels.

There are two new events related to channels, the announce event describing the channel, and the toggle event which notifies any changes in state. Note that all channels are enabled during static initialization in order to guarantee that any event emitted during that phase is recorded.

diff --git a/Engine/Source/Developer/TraceAnalysis/Private/Analysis/Engine.cpp b/Engine/Source/Developer/TraceAnalysis/Private/Analysis/Engine.cpp
--- a/Engine/Source/Developer/TraceAnalysis/Private/Analysis/Engine.cpp
+++ b/Engine/Source/Developer/TraceAnalysis/Private/Analysis/Engine.cpp
@@ -442,6 +442,8 @@ enum ERouteId : uint16
 	RouteId_NewEvent,
 	RouteId_NewTrace,
 	RouteId_Timing,
+	RouteId_ChannelAnnounce,
+	RouteId_ChannelToggle,
 };
 
 ////////////////////////////////////////////////////////////////////////////////
@@ -463,6 +465,8 @@ FAnalysisEngine::FAnalysisEngine(TArray<IAnalyzer*>&& InAnalyzers)
 	AddRoute(SelfIndex, RouteId_NewEvent, RouteHash_NewEvent);
 	AddRoute(SelfIndex, RouteId_NewTrace, "$Trace", "NewTrace");
 	AddRoute(SelfIndex, RouteId_Timing, "$Trace", "Timing");
+	AddRoute(SelfIndex, RouteId_ChannelAnnounce, "$Trace", "ChannelAnnounce");
+	AddRoute(SelfIndex, RouteId_ChannelToggle, "$Trace", "ChannelToggle");
 }
 
 ////////////////////////////////////////////////////////////////////////////////
@@ -533,9 +537,11 @@ bool FAnalysisEngine::OnEvent(uint16 RouteId, const FOnEventContext& Context)
 {
 	switch (RouteId)
 	{
-	case RouteId_NewEvent:	OnNewEventInternal(Context);	break;
-	case RouteId_NewTrace:	OnNewTrace(Context);			break;
-	case RouteId_Timing:	OnTiming(Context);				break;
+	case RouteId_NewEvent:			OnNewEventInternal(Context);		break;
+	case RouteId_NewTrace:			OnNewTrace(Context);				break;
+	case RouteId_Timing:			OnTiming(Context);					break;
+	case RouteId_ChannelAnnounce:	OnChannelAnnounceInternal(Context);	break;
+	case RouteId_ChannelToggle:		OnChannelToggleInternal(Context);	break;
 	}
 
 	return true;
@@ -1093,4 +1099,28 @@ int32 FAnalysisEngine::OnDataProtocol2Aux(FStreamReader& Reader, FAuxDataCollect
 	}
 }
 
+////////////////////////////////////////////////////////////////////////////////
+void FAnalysisEngine::OnChannelAnnounceInternal(const FOnEventContext& Context)
+{
+	const ANSICHAR* ChannelName = (ANSICHAR*)Context.EventData.GetAttachment();
+	const uint32 ChannelId = Context.EventData.GetValue<uint32>("Id");
+
+	for (IAnalyzer* Analyzer : Analyzers)
+	{
+		Analyzer->OnChannelAnnounce(ChannelName, ChannelId);
+	}
+}
+
+////////////////////////////////////////////////////////////////////////////////
+void FAnalysisEngine::OnChannelToggleInternal(const FOnEventContext& Context)
+{
+	const uint32 ChannelId = Context.EventData.GetValue<uint32>("Id");
+	const bool bEnabled = Context.EventData.GetValue<bool>("IsEnabled");
+
+	for (IAnalyzer* Analyzer : Analyzers)
+	{
+		Analyzer->OnChannelToggle(ChannelId, bEnabled);
+	}
+}
+
 } // namespace Trace
diff --git a/Engine/Source/Developer/TraceAnalysis/Private/Analysis/Engine.h b/Engine/Source/Developer/TraceAnalysis/Private/Analysis/Engine.h
--- a/Engine/Source/Developer/TraceAnalysis/Private/Analysis/Engine.h
+++ b/Engine/Source/Developer/TraceAnalysis/Private/Analysis/Engine.h
@@ -42,6 +42,8 @@ private:
 	void				OnNewEventInternal(const FOnEventContext& Context);
 	void				OnNewEventProtocol0(FDispatchBuilder& Builder, const void* EventData);
 	void				OnNewEventProtocol1(FDispatchBuilder& Builder, const void* EventData);
+	void				OnChannelAnnounceInternal(const FOnEventContext& Context);
+	void				OnChannelToggleInternal(const FOnEventContext& Context);
 
 	bool				EstablishTransport(FStreamReader& Reader);
 	bool				OnDataProtocol0();
diff --git a/Engine/Source/Developer/TraceAnalysis/Private/ControlClient.cpp b/Engine/Source/Developer/TraceAnalysis/Private/ControlClient.cpp
--- a/Engine/Source/Developer/TraceAnalysis/Private/ControlClient.cpp
+++ b/Engine/Source/Developer/TraceAnalysis/Private/ControlClient.cpp
@@ -127,6 +127,12 @@ void FControlClient::SendToggleEvent(const TCHAR* EventSpec, bool bState)
     FormatAndSend(TEXT("ToggleEvent %s %d"), EventSpec, bState);
 }
 
+////////////////////////////////////////////////////////////////////////////////
+void FControlClient::SendToggleChannel(const TCHAR* Channels, bool bState /*= true*/)
+{
+	FormatAndSend(TEXT("ToggleChannels %s %d"), Channels, bState);
+}
+
 ////////////////////////////////////////////////////////////////////////////////
 void FControlClient::Send(const TCHAR* Command)
 {
@@ -172,4 +178,4 @@ void FControlClient::Send(const uint8* Data, int Length)
     }
 }
 
-} // namesapce Trace
+} // namespace Trace
diff --git a/Engine/Source/Developer/TraceAnalysis/Private/Recorder/Recorder.cpp b/Engine/Source/Developer/TraceAnalysis/Private/Recorder/Recorder.cpp
--- a/Engine/Source/Developer/TraceAnalysis/Private/Recorder/Recorder.cpp
+++ b/Engine/Source/Developer/TraceAnalysis/Private/Recorder/Recorder.cpp
@@ -34,6 +34,7 @@ private:
 	virtual uint32		GetSessionCount() const override;
 	virtual void		GetActiveSessions(TArray<FRecorderSessionInfo>& OutSessions) const override;
 	virtual bool		ToggleEvent(FRecorderSessionHandle RecordingHandle, const TCHAR* LoggerWildcard, bool bState) override;
+	virtual bool		ToggleChannels(FRecorderSessionHandle RecordingHandle, const TCHAR* Channels, bool bState) override;
 	FSession*			AcceptSession(FSocket& Socket);
 	void				CloseSession(FSession& Session);
 	void				ReapDeadSessions();
@@ -350,6 +351,29 @@ bool FRecorder::ToggleEvent(FRecorderSessionHandle SessionHandle, const TCHAR* L
 	return true;
 }
 
+////////////////////////////////////////////////////////////////////////////////
+bool FRecorder::ToggleChannels(FRecorderSessionHandle RecordingHandle, const TCHAR* Channels, bool bState)
+{
+	TSharedPtr<FInternetAddr> ControlClientAddress;
+	{
+		FScopeLock Lock(&SessionsCS);
+		FSession** FindIt = Sessions.Find(RecordingHandle);
+		if (!FindIt)
+		{
+			return false;
+		}
+		ControlClientAddress = (*FindIt)->ControlClientAddress;
+	}
+	FControlClient ControlClient;
+	if (!ControlClient.Connect(*ControlClientAddress))
+	{
+		return false;
+	}
+	ControlClient.SendToggleChannel(Channels, bState);
+	ControlClient.Disconnect();
+	return true;
+}
+
 ////////////////////////////////////////////////////////////////////////////////
 TSharedPtr<IRecorder> Recorder_Create(TSharedRef<IStore> Store)
 {
diff --git a/Engine/Source/Developer/TraceAnalysis/Public/Trace/Analyzer.h b/Engine/Source/Developer/TraceAnalysis/Public/Trace/Analyzer.h
--- a/Engine/Source/Developer/TraceAnalysis/Public/Trace/Analyzer.h
+++ b/Engine/Source/Developer/TraceAnalysis/Public/Trace/Analyzer.h
@@ -194,6 +194,20 @@ public:
 		return true;
 	}
 
+	/** A new channel has been announced.
+	 * @param ChannelName Clear text name of the channel
+	 * @param ChannelId Unique identifier for the channel */
+	virtual void OnChannelAnnounce(const ANSICHAR* ChannelName, uint32 ChannelId)
+	{
+	}
+
+	/** A channel enabled state has changed.
+	 * @param ChannelId Unique identifier for the channel, see OnChannelAnnounce.
+	 * @param bEnabled True if the channel has been enabled, false otherwise */
+	virtual void OnChannelToggle(uint32 ChannelId, bool bEnabled)
+	{
+	}
+
 private:
 	template <typename ValueType> static ValueType CoerceValue(const void* Addr, int16 SizeAndType);
 };
diff --git a/Engine/Source/Developer/TraceAnalysis/Public/Trace/ControlClient.h b/Engine/Source/Developer/TraceAnalysis/Public/Trace/ControlClient.h
--- a/Engine/Source/Developer/TraceAnalysis/Public/Trace/ControlClient.h
+++ b/Engine/Source/Developer/TraceAnalysis/Public/Trace/ControlClient.h
@@ -50,6 +50,14 @@ public:
 	 * @param EventSpec Specification of the event to toggle.
 	 * @param bState True to enable the event and false to disable it. */
     void SendToggleEvent(const TCHAR* EventSpec, bool bState=true);
+	
+	/**
+	 * Toggles one or more channels on and off. A channel or a comma separated list
+	 * of channels can be controlled.
+	 * @param Channels A single channel name or a comma separated list of channel names.
+	 * @param bState True to enable channel(s), false to disable.
+	 */
+	void SendToggleChannel(const TCHAR* Channels, bool bState = true);
 
 	/** Sends a raw command to the instance.
 	 * @param Command The command to send to the runtime (CRLF terminator not required). */
diff --git a/Engine/Source/Developer/TraceAnalysis/Public/Trace/Recorder.h b/Engine/Source/Developer/TraceAnalysis/Public/Trace/Recorder.h
--- a/Engine/Source/Developer/TraceAnalysis/Public/Trace/Recorder.h
+++ b/Engine/Source/Developer/TraceAnalysis/Public/Trace/Recorder.h
@@ -53,6 +53,16 @@ public:
 	 * @param bState True to enable the event and false to disable it.
 	 * @return True if the command was successfully sent. */
 	virtual bool ToggleEvent(FRecorderSessionHandle RecordingHandle, const TCHAR* EventSpec, bool bState) = 0;
+
+	/**
+	 * Toggles one or more channels on and off. A channel or a comma separated list
+	 * of channels can be controlled.
+	 * @param RecordingHandle Identifier of the active session, from GetActiveSessions().
+	 * @param Channels A single channel name or a comma separated list of channel names.
+	 * @param bState True to enable channel(s), false to disable.
+	 * @return True if the command was successfully sent.
+	 */
+	virtual bool ToggleChannels(FRecorderSessionHandle RecordingHandle, const TCHAR* Channels, bool bState) = 0;
 };
 
 /**
diff --git a/Engine/Source/Developer/TraceServices/Private/AnalysisService.cpp b/Engine/Source/Developer/TraceServices/Private/AnalysisService.cpp
--- a/Engine/Source/Developer/TraceServices/Private/AnalysisService.cpp
+++ b/Engine/Source/Developer/TraceServices/Private/AnalysisService.cpp
@@ -16,6 +16,7 @@
 #include "Model/ThreadsPrivate.h"
 #include "Model/CountersPrivate.h"
 #include "Model/NetProfilerProvider.h"
+#include "Model/Channel.h"
 
 namespace Trace
 {
@@ -174,7 +175,10 @@ void FAnalysisService::AnalyzeInternal(TSharedRef<FAnalysisSession> AnalysisSess
 		FNetProfilerProvider* NetProfilerProvider = new FNetProfilerProvider(Session);
 		Session.AddProvider(FNetProfilerProvider::ProviderName, NetProfilerProvider);
 
-		Session.AddAnalyzer(new FMiscTraceAnalyzer(Session, *ThreadProvider, *BookmarkProvider, *LogProvider, *FrameProvider));
+		FChannelProvider* ChannelProvider = new FChannelProvider();
+		Session.AddProvider(FChannelProvider::ProviderName, ChannelProvider);
+
+		Session.AddAnalyzer(new FMiscTraceAnalyzer(Session, *ThreadProvider, *BookmarkProvider, *LogProvider, *FrameProvider, *ChannelProvider));
 		Session.AddAnalyzer(new FLogTraceAnalyzer(Session, *LogProvider));
 
 		ModuleService.OnAnalysisBegin(Session);
diff --git a/Engine/Source/Developer/TraceServices/Private/Analyzers/MiscTraceAnalysis.cpp b/Engine/Source/Developer/TraceServices/Private/Analyzers/MiscTraceAnalysis.cpp
--- a/Engine/Source/Developer/TraceServices/Private/Analyzers/MiscTraceAnalysis.cpp
+++ b/Engine/Source/Developer/TraceServices/Private/Analyzers/MiscTraceAnalysis.cpp
@@ -5,18 +5,21 @@
 #include "Model/ThreadsPrivate.h"
 #include "Model/BookmarksPrivate.h"
 #include "Model/FramesPrivate.h"
+#include "Model/Channel.h"
 #include "Common/Utils.h"
 
 FMiscTraceAnalyzer::FMiscTraceAnalyzer(Trace::IAnalysisSession& InSession,
 									   Trace::FThreadProvider& InThreadProvider,
 									   Trace::FBookmarkProvider& InBookmarkProvider,
 									   Trace::FLogProvider& InLogProvider,
-									   Trace::FFrameProvider& InFrameProvider)
+									   Trace::FFrameProvider& InFrameProvider,
+									   Trace::FChannelProvider& InChannelProvider)
 	: Session(InSession)
 	, ThreadProvider(InThreadProvider)
 	, BookmarkProvider(InBookmarkProvider)
 	, LogProvider(InLogProvider)
 	, FrameProvider(InFrameProvider)
+	, ChannelProvider(InChannelProvider)
 {
 	Trace::FLogCategory& BookmarkLogCategory = LogProvider.GetCategory(Trace::FLogProvider::ReservedLogCategory_Bookmark);
 	BookmarkLogCategory.Name = TEXT("LogBookmark");
@@ -160,11 +163,23 @@ bool FMiscTraceAnalyzer::OnEvent(uint16 RouteId, const FOnEventContext& Context)
 		}
 		break;
 	}
+	
+
 	}
 
 	return true;
 }
 
+void FMiscTraceAnalyzer::OnChannelAnnounce(const ANSICHAR* ChannelName, uint32 ChannelId)
+{
+	ChannelProvider.AnnounceChannel(ChannelName, ChannelId);
+}
+
+void FMiscTraceAnalyzer::OnChannelToggle(uint32 ChannelId, bool bEnabled)
+{
+	ChannelProvider.UpdateChannel(ChannelId, bEnabled);
+}
+
 FMiscTraceAnalyzer::FThreadState* FMiscTraceAnalyzer::GetThreadState(uint32 ThreadId)
 {
 	if (!ThreadStateMap.Contains(ThreadId))
diff --git a/Engine/Source/Developer/TraceServices/Private/Analyzers/MiscTraceAnalysis.h b/Engine/Source/Developer/TraceServices/Private/Analyzers/MiscTraceAnalysis.h
--- a/Engine/Source/Developer/TraceServices/Private/Analyzers/MiscTraceAnalysis.h
+++ b/Engine/Source/Developer/TraceServices/Private/Analyzers/MiscTraceAnalysis.h
@@ -14,6 +14,7 @@ namespace Trace
 	class FBookmarkProvider;
 	class FLogProvider;
 	class FFrameProvider;
+	class FChannelProvider;
 }
 
 class FMiscTraceAnalyzer
@@ -24,9 +25,12 @@ public:
 					   Trace::FThreadProvider& ThreadProvider,
 					   Trace::FBookmarkProvider& BookmarkProvider,
 					   Trace::FLogProvider& LogProvider,
-					   Trace::FFrameProvider& FrameProvider);
+					   Trace::FFrameProvider& FrameProvider, 
+					   Trace::FChannelProvider& ChannelProvider);
 	virtual void OnAnalysisBegin(const FOnAnalysisContext& Context) override;
 	virtual bool OnEvent(uint16 RouteId, const FOnEventContext& Context) override;
+	virtual void OnChannelAnnounce(const ANSICHAR* ChannelName, uint32 ChannelId) override;
+	virtual void OnChannelToggle(uint32 ChannelId, bool bEnabled) override;
 
 private:
 	enum : uint16
@@ -58,6 +62,7 @@ private:
 	Trace::FBookmarkProvider& BookmarkProvider;
 	Trace::FLogProvider& LogProvider;
 	Trace::FFrameProvider& FrameProvider;
+	Trace::FChannelProvider& ChannelProvider;
 	TMap<uint32, TSharedRef<FThreadState>> ThreadStateMap;
 	uint64 LastFrameCycle[TraceFrameType_Count] = { 0, 0 };
 };
diff --git a/Engine/Source/Developer/TraceServices/Private/Model/Channel.cpp b/Engine/Source/Developer/TraceServices/Private/Model/Channel.cpp
new file mode 100644
--- /dev/null
+++ b/Engine/Source/Developer/TraceServices/Private/Model/Channel.cpp
@@ -0,0 +1,49 @@
+#include "Channel.h"
+
+namespace Trace {
+
+const FName FChannelProvider::ProviderName("ChannelProvider");
+
+///////////////////////////////////////////////////////////////////////////////
+FChannelProvider::FChannelProvider()
+{
+}
+
+///////////////////////////////////////////////////////////////////////////////
+void FChannelProvider::AnnounceChannel(const ANSICHAR* InChannelName, uint32 Id)
+{
+	FString ChannelName(ANSI_TO_TCHAR(InChannelName));
+	ChannelName.GetCharArray()[0] = TChar<TCHAR>::ToUpper(ChannelName.GetCharArray()[0]);
+	Channels.Add(FChannelEntry{
+		Id,
+		ChannelName,
+		true
+	});
+}
+
+///////////////////////////////////////////////////////////////////////////////
+void FChannelProvider::UpdateChannel(uint32 Id, bool bEnabled)
+{
+	const auto FoundEntry = Channels.FindByPredicate([Id](const FChannelEntry& Entry) {
+		return Entry.Id == Id;
+	});
+
+	if (FoundEntry)
+	{
+		FoundEntry->bIsEnabled = bEnabled;
+	}
+}
+
+///////////////////////////////////////////////////////////////////////////////
+uint64 FChannelProvider::GetChannelCount() const
+{
+	return Channels.Num();
+}
+
+///////////////////////////////////////////////////////////////////////////////
+const TArray<FChannelEntry>& FChannelProvider::GetChannels() const
+{
+	return Channels;
+}
+
+}
\ No newline at end of file
diff --git a/Engine/Source/Developer/TraceServices/Private/Model/Channel.h b/Engine/Source/Developer/TraceServices/Private/Model/Channel.h
new file mode 100644
--- /dev/null
+++ b/Engine/Source/Developer/TraceServices/Private/Model/Channel.h
@@ -0,0 +1,25 @@
+// Copyright 1998-2019 Epic Games, Inc. All Rights Reserved.
+
+#pragma once
+
+#include "TraceServices/Model/Channel.h"
+
+namespace Trace {
+
+class FChannelProvider : public IChannelProvider
+{
+public:
+	static const FName ProviderName;
+
+	FChannelProvider();
+	void AnnounceChannel(const ANSICHAR* ChannelName, uint32 Id);
+	void UpdateChannel(uint32 Id, bool bEnabled);
+
+	virtual uint64	GetChannelCount() const override;
+	virtual const TArray<FChannelEntry>& GetChannels() const override;
+
+private:
+	TArray<FChannelEntry> Channels;
+};
+
+} // namespace Trace
\ No newline at end of file
diff --git a/Engine/Source/Developer/TraceServices/Private/Model/Frames.cpp b/Engine/Source/Developer/TraceServices/Private/Model/Frames.cpp
--- a/Engine/Source/Developer/TraceServices/Private/Model/Frames.cpp
+++ b/Engine/Source/Developer/TraceServices/Private/Model/Frames.cpp
@@ -114,6 +114,11 @@ void FFrameProvider::BeginFrame(ETraceFrameType FrameType, double Time)
 void FFrameProvider::EndFrame(ETraceFrameType FrameType, double Time)
 {
 	Session.WriteAccessCheck();
+	// If the EndFrame event is the first event that comes through
+	if (Frames[FrameType].Num() == 0)
+	{
+		return;
+	}
 	FFrame& Frame = Frames[FrameType][Frames[FrameType].Num() - 1];
 	Frame.EndTime = Time;
 	Session.UpdateDurationSeconds(Time);
diff --git a/Engine/Source/Developer/TraceServices/Private/SessionService.cpp b/Engine/Source/Developer/TraceServices/Private/SessionService.cpp
--- a/Engine/Source/Developer/TraceServices/Private/SessionService.cpp
+++ b/Engine/Source/Developer/TraceServices/Private/SessionService.cpp
@@ -232,6 +232,17 @@ bool FSessionService::IsModuleEnabled(Trace::FSessionHandle SessionHandle, const
 	return FindIt->EnabledModuleLoggersMap.Contains(ModuleName);
 }
 
+bool FSessionService::ToggleChannels(Trace::FSessionHandle SessionHandle, const TCHAR* Channels, bool bState)
+{
+	FScopeLock Lock(&SessionsCS);
+	const FSessionInfoInternal* FindIt = Sessions.Find(SessionHandle);
+	if (!FindIt)
+	{
+		return false;
+	}
+	return TraceRecorder->ToggleChannels(FindIt->RecorderSessionHandle, Channels, bState);
+}
+
 bool FSessionService::ConnectSession(const TCHAR* ControlClientAddress)
 {
 	ISocketSubsystem* Sockets = ISocketSubsystem::Get();
diff --git a/Engine/Source/Developer/TraceServices/Private/SessionServicePrivate.h b/Engine/Source/Developer/TraceServices/Private/SessionServicePrivate.h
--- a/Engine/Source/Developer/TraceServices/Private/SessionServicePrivate.h
+++ b/Engine/Source/Developer/TraceServices/Private/SessionServicePrivate.h
@@ -30,6 +30,7 @@ public:
 	virtual bool GetSessionInfo(Trace::FSessionHandle SessionHandle, Trace::FSessionInfo& OutSessionInfo) const override;
 	virtual void SetModuleEnabled(Trace::FSessionHandle SessionHandle, const FName& ModuleName, bool bState) override;
 	virtual bool IsModuleEnabled(Trace::FSessionHandle SessionHandle, const FName& ModuleName) const override;
+	virtual bool ToggleChannels(Trace::FSessionHandle SessionHandle, const TCHAR* Channels, bool bState) override;
 	virtual bool ConnectSession(const TCHAR* ControlClientAddress) override;
 	virtual TSharedPtr<const IAnalysisSession> StartAnalysis(FSessionHandle SessionHandle) override;
 
diff --git a/Engine/Source/Developer/TraceServices/Public/TraceServices/Model/Channel.h b/Engine/Source/Developer/TraceServices/Public/TraceServices/Model/Channel.h
new file mode 100644
--- /dev/null
+++ b/Engine/Source/Developer/TraceServices/Public/TraceServices/Model/Channel.h
@@ -0,0 +1,25 @@
+// Copyright 1998-2019 Epic Games, Inc. All Rights Reserved.
+
+#pragma once
+
+#include "TraceServices/Model/AnalysisSession.h"
+
+namespace Trace {
+
+struct FChannelEntry
+{
+	uint32 Id;
+	FString Name;
+	bool bIsEnabled;
+};
+
+class IChannelProvider
+	: public IProvider
+{
+public:
+	virtual ~IChannelProvider() = default;
+	virtual uint64	GetChannelCount() const = 0;
+	virtual const TArray<FChannelEntry>& GetChannels() const = 0;
+};
+
+} // namespace Trace
\ No newline at end of file
diff --git a/Engine/Source/Developer/TraceServices/Public/TraceServices/SessionService.h b/Engine/Source/Developer/TraceServices/Public/TraceServices/SessionService.h
--- a/Engine/Source/Developer/TraceServices/Public/TraceServices/SessionService.h
+++ b/Engine/Source/Developer/TraceServices/Public/TraceServices/SessionService.h
@@ -41,6 +41,7 @@ public:
 	virtual bool GetSessionInfo(FSessionHandle SessionHandle, FSessionInfo& OutSessionInfo) const = 0;
 	virtual void SetModuleEnabled(Trace::FSessionHandle SessionHandle, const FName& ModuleName, bool bState) = 0;
 	virtual bool IsModuleEnabled(Trace::FSessionHandle SessionHandle, const FName& ModuleName) const = 0;
+	virtual bool ToggleChannels(Trace::FSessionHandle SessionHandle, const TCHAR* Channels, bool bState) = 0;
 	virtual bool ConnectSession(const TCHAR* ControlClientAddress) = 0;
 	virtual TSharedPtr<const IAnalysisSession> StartAnalysis(FSessionHandle SessionHandle) = 0;
 };
diff --git a/Engine/Source/Runtime/TraceLog/Private/Trace/Channel.cpp b/Engine/Source/Runtime/TraceLog/Private/Trace/Channel.cpp
new file mode 100644
--- /dev/null
+++ b/Engine/Source/Runtime/TraceLog/Private/Trace/Channel.cpp
@@ -0,0 +1,168 @@
+// Copyright 1998-2019 Epic Games, Inc. All Rights Reserved.
+#include "Trace/Detail/Channel.h"
+#include "Trace/Trace.h"
+#include "Trace/Detail/Atomic.h"
+#include "Trace/Detail/Channel.h"
+#include "Misc/Parse.h"
+
+#if UE_TRACE_ENABLED
+
+namespace Trace
+{
+
+///////////////////////////////////////////////////////////////////////////////
+static FChannel* volatile GHeadChannel; // = nullptr;
+static const size_t ChannelNameMaxLength = 64u;
+
+///////////////////////////////////////////////////////////////////////////////
+template <int DestSize>
+static void ToAnsiCheap(ANSICHAR(&Dest)[DestSize], const WIDECHAR* Src)
+{
+	for (ANSICHAR& Out : Dest)
+	{
+		Out = ANSICHAR(*Src++ & 0x7f);
+		if (Out == '\0')
+		{
+			break;
+		}
+	}
+};
+
+////////////////////////////////////////////////////////////////////////////////
+template <typename ElementType>
+static uint32 ChannelGetHash(const ElementType* Input, int32 Length = -1)
+{
+	uint32 Result = 0x811c9dc5;
+	for (; *Input && Length; ++Input, --Length)
+	{
+		Result ^= *Input;
+		Result *= 0x01000193;
+	}
+	return Result;
+}
+
+////////////////////////////////////////////////////////////////////////////////
+static uint32 ChannelGetHash(uint32 LoggerHash, uint32 NameHash)
+{
+	uint32 Parts[3] = { LoggerHash, NameHash, 0 };
+	return ChannelGetHash(Parts);
+}
+
+///////////////////////////////////////////////////////////////////////////////
+static uint32 GetChannelIdentifier(ANSICHAR(&Buffer) [ChannelNameMaxLength], const ANSICHAR* ChannelName)
+{
+	// Strip "Channel" suffix if it exists
+	if (const ANSICHAR* ChannelStr = strstr(ChannelName, "Channel"))
+	{
+		const size_t Count = ChannelStr - ChannelName;
+		if (Count > 0 && strlen(ChannelName) > Count)
+		{
+			strncpy_s(Buffer, ChannelNameMaxLength, ChannelName, Count);
+			Buffer[Count] = '\0';
+		}
+	}
+	else
+	{
+		strcpy_s(Buffer, ChannelNameMaxLength, ChannelName);
+	}
+
+	// Convert to lower case
+	const uint32 ChannelNameLen = strlen(Buffer);
+	for (size_t i = 0; i < ChannelNameLen; ++i)
+	{
+		Buffer[i] = tolower(Buffer[i]);
+	}
+
+	return ChannelNameLen;
+}
+
+///////////////////////////////////////////////////////////////////////////////
+void FChannel::Register(FChannel& Channel, const ANSICHAR* ChannelName)
+{
+	UE_TRACE_EVENT_BEGIN($Trace, ChannelAnnounce, Important)
+		UE_TRACE_EVENT_FIELD(uint32, Id)
+	UE_TRACE_EVENT_END()
+
+	
+	ANSICHAR ChannelIdentifier[ChannelNameMaxLength];
+	const uint32 ChannelNameLen = GetChannelIdentifier(ChannelIdentifier, ChannelName);
+	const uint32 ChannelNameHash = ChannelGetHash(ChannelIdentifier);
+
+	Channel.ChannelNameHash = ChannelNameHash;
+
+	// Append channel to linked list
+	for (;; Private::PlatformYield())
+	{
+		FChannel* HeadChannel = Private::AtomicLoadRelaxed(&GHeadChannel);
+		Channel.Handle = HeadChannel;
+		if (Private::AtomicCompareExchangeRelease(&GHeadChannel, &Channel, HeadChannel))
+		{
+			break;
+		}
+	}
+
+	UE_TRACE_LOG($Trace, ChannelAnnounce, TraceLogChannel, ChannelNameLen+1)
+		<< ChannelAnnounce.Id(ChannelNameHash)
+		<< ChannelAnnounce.Attachment(ChannelIdentifier, ChannelNameLen+1);
+}
+
+///////////////////////////////////////////////////////////////////////////////
+ void FChannel::ToggleAll(bool bEnabled)
+{
+	FChannel* Channel = Private::AtomicLoadAcquire(&GHeadChannel);
+	for (; Channel != nullptr; Channel = (FChannel*)(Channel->Handle))
+	{
+		FChannel::Toggle(Channel, bEnabled);
+	}
+}
+
+ ///////////////////////////////////////////////////////////////////////////////
+ bool FChannel::Toggle(FChannel* Channel, bool bEnabled)
+ {
+	UE_TRACE_EVENT_BEGIN($Trace, ChannelToggle, Important)
+		UE_TRACE_EVENT_FIELD(uint32, Id)
+		UE_TRACE_EVENT_FIELD(bool, IsEnabled)
+	UE_TRACE_EVENT_END()
+
+	const bool bWasEnabled = !Channel->bDisabled;
+	if (bWasEnabled != bEnabled)
+	{
+		Channel->bDisabled = !bEnabled;
+		UE_TRACE_LOG($Trace, ChannelToggle, TraceLogChannel)
+			<< ChannelToggle.Id(Channel->ChannelNameHash)
+			<< ChannelToggle.IsEnabled(bEnabled);
+	}
+	return bWasEnabled;
+ }
+
+ ///////////////////////////////////////////////////////////////////////////////
+ bool FChannel::Toggle(const ANSICHAR* ChannelName, bool bEnabled)
+ {
+	ANSICHAR ChannelIdentifier[ChannelNameMaxLength];
+	const uint32 ChannelNameLen = GetChannelIdentifier(ChannelIdentifier, ChannelName);
+	const uint32 ChannelNameHash = ChannelGetHash(ChannelIdentifier);
+
+	FChannel* Channel = Private::AtomicLoadAcquire(&GHeadChannel);
+	for (; Channel != nullptr; Channel = (FChannel*)(Channel->Handle))
+	{
+		if (Channel->ChannelNameHash == ChannelNameHash)
+		{
+			return FChannel::Toggle(Channel, bEnabled);
+		}
+	}
+
+	return false;
+ }
+
+ ///////////////////////////////////////////////////////////////////////////////
+ bool FChannel::Toggle(const TCHAR* ChannelName, bool bEnabled)
+ {
+	 ANSICHAR ChannelNameA[ChannelNameMaxLength];
+	 ToAnsiCheap(ChannelNameA, ChannelName);
+
+	 return FChannel::Toggle(ChannelNameA, bEnabled);
+ }
+
+}
+
+#endif //UE_TRACE_ENABLED
\ No newline at end of file
diff --git a/Engine/Source/Runtime/TraceLog/Private/Trace/Trace.cpp b/Engine/Source/Runtime/TraceLog/Private/Trace/Trace.cpp
--- a/Engine/Source/Runtime/TraceLog/Private/Trace/Trace.cpp
+++ b/Engine/Source/Runtime/TraceLog/Private/Trace/Trace.cpp
@@ -35,6 +35,13 @@ static void ToAnsiCheap(ANSICHAR (&Dest)[DestSize], const WIDECHAR* Src)
 	}
 };
 
+////////////////////////////////////////////////////////////////////////////////
+bool Initialize()
+{
+	FChannel::ToggleAll(false);
+	return true;
+}
+
 ////////////////////////////////////////////////////////////////////////////////
 bool SendTo(const TCHAR* InHost)
 {
@@ -60,6 +67,21 @@ uint32 ToggleEvent(const TCHAR* Wildcard, bool bState)
 	return Private::Writer_EventToggle(WildcardA, bState);
 }
 
+////////////////////////////////////////////////////////////////////////////////
+bool ToggleChannel(const TCHAR* ChannelName, bool bEnabled)
+{
+	ANSICHAR ChannelNameA[64];
+	ToAnsiCheap(ChannelNameA, ChannelName);
+
+	return FChannel::Toggle(ChannelNameA, bEnabled);
+}
+
+////////////////////////////////////////////////////////////////////////////////
+bool ToggleChannel(struct FChannel& Channel, bool bEnabled)
+{
+	return FChannel::Toggle(&Channel, bEnabled);
+}
+
 } // namespace Trace
 
 #else
diff --git a/Engine/Source/Runtime/TraceLog/Private/Trace/Writer.cpp b/Engine/Source/Runtime/TraceLog/Private/Trace/Writer.cpp
--- a/Engine/Source/Runtime/TraceLog/Private/Trace/Writer.cpp
+++ b/Engine/Source/Runtime/TraceLog/Private/Trace/Writer.cpp
@@ -6,20 +6,14 @@
 
 #include "Trace/Detail/Atomic.h"
 #include "Trace/Detail/Protocol.h"
+#include "Trace/Detail/Channel.h"
 #include "Trace/Platform.h"
 #include "Trace/Trace.h"
 
 #include "Misc/CString.h"
 #include "Templates/UnrealTemplate.h"
 
-#if PLATFORM_CPU_X86_FAMILY
-	#include <emmintrin.h>
-	#define PLATFORM_YIELD()	_mm_pause()
-#elif PLATFORM_CPU_ARM_FAMILY
-	#define PLATFORM_YIELD()	__builtin_arm_yield();
-#else
-	#error Unsupported architecture!
-#endif
+
 
 namespace Trace {
 namespace Private {
@@ -42,16 +36,6 @@ UE_TRACE_EVENT_BEGIN($Trace, Memory, Always)
 UE_TRACE_EVENT_END()
 #endif // TRACE_PRIVATE_PERF
 
-
-
-////////////////////////////////////////////////////////////////////////////////
-inline void Writer_Yield()
-{
-	PLATFORM_YIELD();
-}
-
-
-
 ////////////////////////////////////////////////////////////////////////////////
 static uint64 GStartCycle = 0;
 
@@ -66,12 +50,12 @@ void Writer_InitializeTiming()
 {
 	GStartCycle = TimeGetTimestamp();
 
-	UE_TRACE_EVENT_BEGIN($Trace, Timing, Always|Important)
+	UE_TRACE_EVENT_BEGIN($Trace, Timing, Important)
 		UE_TRACE_EVENT_FIELD(uint64, StartCycle)
 		UE_TRACE_EVENT_FIELD(uint64, CycleFrequency)
 	UE_TRACE_EVENT_END()
 
-	UE_TRACE_LOG($Trace, Timing)
+	UE_TRACE_LOG($Trace, Timing, TraceLogChannel)
 		<< Timing.StartCycle(GStartCycle)
 		<< Timing.CycleFrequency(TimeGetFrequency());
 }
@@ -99,7 +83,7 @@ FWriteTlsContext::FWriteTlsContext()
 	Buffer = Target;
 
 	// Assign a new id to this thread.
-	for (;; Writer_Yield())
+	for (;; Private::PlatformYield())
 	{
 		UPTRINT CandidateId = UPTRINT(AtomicLoadRelaxed((void* volatile*)&ThreadIdCounter));
 		UPTRINT NextId = CandidateId + 1;
@@ -176,7 +160,7 @@ static FWriteBuffer* Writer_NextBufferInternal(uint32 PageGrowth)
 		{
 			if (!AtomicCompareExchangeRelaxed(&GPoolFreeList, Owned->Next, Owned))
 			{
-				Writer_Yield();
+				Private::PlatformYield();
 				continue;
 			}
 		}
@@ -194,7 +178,7 @@ static FWriteBuffer* Writer_NextBufferInternal(uint32 PageGrowth)
 		{
 			// Someone else is mapping memory so we'll briefly yield and try the
 			// free list again.
-			Writer_Yield();
+			Private::PlatformYield();
 			continue;
 		}
 
@@ -218,7 +202,7 @@ static FWriteBuffer* Writer_NextBufferInternal(uint32 PageGrowth)
 		}
 
 		// And insert the block list into the freelist. 'Block' is now the last block
-		for (auto* ListNode = (FWriteBuffer*)Block;; Writer_Yield())
+		for (auto* ListNode = (FWriteBuffer*)Block;; Private::PlatformYield())
 		{
 			ListNode->Next = AtomicLoadRelaxed(&GPoolFreeList);
 			if (AtomicCompareExchangeRelease(&GPoolFreeList, (FWriteBuffer*)FirstBlock, ListNode->Next))
@@ -237,7 +221,7 @@ static FWriteBuffer* Writer_NextBufferInternal(uint32 PageGrowth)
 	NextBuffer->EtxOffset = 0;
 
 	// Add this next buffer to the active list.
-	for (;; Writer_Yield())
+	for (;; Private::PlatformYield())
 	{
 		NextBuffer->Next = AtomicLoadRelaxed(&GNextBufferList);
 		if (AtomicCompareExchangeRelease(&GNextBufferList, NextBuffer, NextBuffer->Next))
@@ -404,7 +388,7 @@ static void Writer_ConsumeEvents()
 
 	// Claim ownership of the latest chain of sent events.
 	FWriteBuffer* __restrict NextBufferList;
-	for (;; Writer_Yield())
+	for (;; Private::PlatformYield())
 	{
 		NextBufferList = AtomicLoadRelaxed(&GNextBufferList);
 		if (AtomicCompareExchangeAcquire(&GNextBufferList, (FWriteBuffer*)nullptr, NextBufferList))
@@ -592,18 +576,18 @@ static void Writer_ConsumeEvents()
 	}
 
 #if TRACE_PRIVATE_PERF
-	UE_TRACE_LOG($Trace, WorkerThread)
+	UE_TRACE_LOG($Trace, WorkerThread, TraceLogChannel)
 		<< WorkerThread.Cycles(uint32(TimeGetTimestamp() - StartTsc))
 		<< WorkerThread.BytesSent(BytesSent);
 
-	UE_TRACE_LOG($Trace, Memory)
+	UE_TRACE_LOG($Trace, Memory, TraceLogChannel)
 		<< Memory.AllocSize(uint32(GPoolPageCursor - GPoolBase));
 #endif // TRACE_PRIVATE_PERF
 
 	// Put the retirees we found back into the system again.
 	if (RetireList.Head != nullptr)
 	{
-		for (FWriteBuffer* ListNode = RetireList.Tail;; Writer_Yield())
+		for (FWriteBuffer* ListNode = RetireList.Tail;; Private::PlatformYield())
 		{
 			ListNode->Next = AtomicLoadRelaxed(&GPoolFreeList);
 			if (AtomicCompareExchangeRelease(&GPoolFreeList, RetireList.Head, ListNode->Next))
@@ -675,7 +659,7 @@ enum class EControlState : uint8
 
 struct FControlCommands
 {
-	enum { Max = 3 };
+	enum { Max = 16 };
 	struct
 	{
 		uint32	Hash;
@@ -938,6 +922,28 @@ static void Writer_InitializeControl()
 			Writer_EventToggle(Wildcard, State[0] != '0');
 		}
 	);
+
+	Writer_ControlAddCommand("ToggleChannels", nullptr, 
+		[] (void*, uint32 ArgC, ANSICHAR const* const* ArgV) 
+		{
+			if (ArgC < 2)
+			{
+				return;
+			}
+
+			const size_t BufferSize = 512;
+			ANSICHAR Channels[BufferSize] = {};
+			ANSICHAR* Ctx;
+			const bool bState = (ArgV[1][0] != '0');
+			strcpy_s(Channels, BufferSize, ArgV[0]);
+			ANSICHAR* Channel = FPlatformString::Strtok(Channels, ",", &Ctx);
+			while (Channel)
+			{
+				FChannel::Toggle(Channel, bState);
+				Channel = FPlatformString::Strtok(nullptr, ",", &Ctx);
+			}
+		}
+	);
 }
 
 ////////////////////////////////////////////////////////////////////////////////
@@ -979,13 +985,13 @@ static bool GInitialized = false;
 ////////////////////////////////////////////////////////////////////////////////
 static void Writer_LogHeader()
 {
-	UE_TRACE_EVENT_BEGIN($Trace, NewTrace, Always|Important)
+	UE_TRACE_EVENT_BEGIN($Trace, NewTrace, Important)
 		UE_TRACE_EVENT_FIELD(uint16, Endian)
 		UE_TRACE_EVENT_FIELD(uint8, Version)
 		UE_TRACE_EVENT_FIELD(uint8, PointerSize)
 	UE_TRACE_EVENT_END()
 
-	UE_TRACE_LOG($Trace, NewTrace)
+	UE_TRACE_LOG($Trace, NewTrace, TraceLogChannel)
 		<< NewTrace.Version(2)
 		<< NewTrace.Endian(0x524d)
 		<< NewTrace.PointerSize(sizeof(void*));
@@ -1123,7 +1129,7 @@ void Writer_EventCreate(
 
 	// Assign a unique ID for this event
 	uint32 Uid;
-	for (;; Writer_Yield())
+	for (;; Private::PlatformYield())
 	{
 		UPTRINT CurrentUid = UPTRINT(AtomicLoadRelaxed((void* volatile*)&GEventUidCounter));
 		UPTRINT NextUid = CurrentUid + 1;
@@ -1150,9 +1156,9 @@ void Writer_EventCreate(
  	Target->Uid = uint16(Uid);
 	Target->LoggerHash = LoggerHash;
 	Target->Hash = Writer_EventGetHash(LoggerHash, NameHash);
-	Target->Enabled.Internal = !!(Flags & FEventDef::Flag_Always);
 	Target->Enabled.bOptedIn = false;
 	Target->bInitialized = true;
+	Target->bImportant = Flags & FEventDef::Flag_Important;
 
 	// Calculate the number of fields and size of name data.
 	int NamesSize = LoggerName.Length + EventName.Length;
@@ -1217,7 +1223,7 @@ void Writer_EventCreate(
 	Writer_EndLog(LogInstance);
 
 	// Add this new event into the list so we can look them up later.
-	for (;; Writer_Yield())
+	for (;; Private::PlatformYield())
 	{
 		FEventDef* HeadEvent = AtomicLoadRelaxed(&GHeadEvent);
 		Target->Handle = HeadEvent;
diff --git a/Engine/Source/Runtime/TraceLog/Public/Trace/Detail/Atomic.h b/Engine/Source/Runtime/TraceLog/Public/Trace/Detail/Atomic.h
--- a/Engine/Source/Runtime/TraceLog/Public/Trace/Detail/Atomic.h
+++ b/Engine/Source/Runtime/TraceLog/Public/Trace/Detail/Atomic.h
@@ -14,10 +14,16 @@ template <typename Type> bool	AtomicCompareExchangeRelaxed(Type volatile* Target
 template <typename Type> bool	AtomicCompareExchangeAcquire(Type volatile* Target, Type New, Type Expected);
 template <typename Type> bool	AtomicCompareExchangeRelease(Type volatile* Target, Type New, Type Expected);
 uint32							AtomicIncrementRelaxed(uint32 volatile* Target);
+void							PlatformYield();
 
 } // namespace Private
 } // namespace Trace
 
+////////////////////////////////////////////////////////////////////////////////
+#if PLATFORM_CPU_X86_FAMILY
+#include <emmintrin.h>
+#endif
+
 ////////////////////////////////////////////////////////////////////////////////
 #define IS_MSVC					0
 #define IS_GCC_COMPATIBLE		0
@@ -44,6 +50,18 @@ uint32							AtomicIncrementRelaxed(uint32 volatile* Target);
 namespace Trace {
 namespace Private {
 
+////////////////////////////////////////////////////////////////////////////////
+inline void PlatformYield()
+{
+#if PLATFORM_CPU_X86_FAMILY
+	_mm_pause();
+#elif PLATFORM_CPU_ARM_FAMILY
+	__builtin_arm_yield();
+#else
+	#error Unsupported architecture!
+#endif
+}
+
 ////////////////////////////////////////////////////////////////////////////////
 #if IS_MSVC
 #	if defined(_M_ARM) || defined(_M_ARM64)
diff --git a/Engine/Source/Runtime/TraceLog/Public/Trace/Detail/Channel.h b/Engine/Source/Runtime/TraceLog/Public/Trace/Detail/Channel.h
new file mode 100644
--- /dev/null
+++ b/Engine/Source/Runtime/TraceLog/Public/Trace/Detail/Channel.h
@@ -0,0 +1,43 @@
+// Copyright 1998-2019 Epic Games, Inc. All Rights Reserved.
+
+#pragma once
+
+#include "Trace/Config.h"
+
+#if UE_TRACE_ENABLED
+
+namespace Trace
+{
+
+/*
+	A named channel which can be used to filter trace events. Channels can be 
+	combined using the '|' operator which allows expressions like
+
+	```
+	UE_TRACE_LOG(FooWriter, FooEvent, FooChannel|BarChannel);
+	```
+
+	Channels are by default enabled until this method is called. This is to allow
+	events to be emitted during static initialization. In fact all events during
+	this phase are always emitted. In this method we disable all channels except
+	those specified on the command line using -tracechannels argument.
+*/
+struct FChannel 
+{
+	TRACELOG_API static void Register(FChannel& Channel, const ANSICHAR* ChannelName);
+	TRACELOG_API static bool Toggle(FChannel* Channel, bool bEnabled);
+	TRACELOG_API static bool Toggle(const ANSICHAR* ChannelName, bool bEnabled);
+	TRACELOG_API static bool Toggle(const TCHAR* ChannelName, bool bEnabled);
+	TRACELOG_API static void ToggleAll(bool bEnabled);
+	bool IsEnabled() const;
+	explicit operator bool() const;
+	bool operator|(const FChannel& Rhs) const;
+
+	void*			Handle;
+	uint32			ChannelNameHash;
+	bool			bDisabled;
+};
+
+}
+
+#endif //UE_TRACE_ENABLED
\ No newline at end of file
diff --git a/Engine/Source/Runtime/TraceLog/Public/Trace/Detail/Channel.inl b/Engine/Source/Runtime/TraceLog/Public/Trace/Detail/Channel.inl
new file mode 100644
--- /dev/null
+++ b/Engine/Source/Runtime/TraceLog/Public/Trace/Detail/Channel.inl
@@ -0,0 +1,43 @@
+// Copyright 1998-2019 Epic Games, Inc. All Rights Reserved.
+
+#pragma once
+
+#include "Channel.h"
+
+#if UE_TRACE_ENABLED
+
+namespace Trace {
+
+////////////////////////////////////////////////////////////////////////////////
+inline bool FChannel::IsEnabled() const
+{
+	return !bDisabled;
+}
+
+////////////////////////////////////////////////////////////////////////////////
+
+inline FChannel::operator bool() const
+{
+	return IsEnabled();
+}
+
+////////////////////////////////////////////////////////////////////////////////
+inline bool FChannel::operator|(const FChannel& Rhs) const
+{
+	return IsEnabled() && Rhs.IsEnabled();
+}
+
+
+////////////////////////////////////////////////////////////////////////////////
+struct FTraceChannel : public FChannel
+{
+	bool IsEnabled() const { return true; }
+	explicit operator bool() const { return true; }
+};
+
+}
+
+// General trace channel. Used by all built in events.
+static Trace::FTraceChannel TraceLogChannel;
+
+#endif //UE_TRACE_ENABLED
\ No newline at end of file
diff --git a/Engine/Source/Runtime/TraceLog/Public/Trace/Detail/EventDef.h b/Engine/Source/Runtime/TraceLog/Public/Trace/Detail/EventDef.h
--- a/Engine/Source/Runtime/TraceLog/Public/Trace/Detail/EventDef.h
+++ b/Engine/Source/Runtime/TraceLog/Public/Trace/Detail/EventDef.h
@@ -20,9 +20,8 @@ class FEventDef
 public:
 	enum
 	{
-		Flag_Always			= 1 << 0,
-		Flag_Important		= 1 << 1,
-		Flag_MaybeHasAux	= 1 << 2,
+		Flag_Important		= 1 << 0,
+		Flag_MaybeHasAux	= 1 << 1,
 	};
 
 	class FLogScope
@@ -49,6 +48,7 @@ public:
 		uint16					Test;
 	}							Enabled;
 	bool						bInitialized;
+	bool						bImportant;
 	TRACELOG_API static void	Create(FEventDef* Target, const FLiteralName& LoggerName, const FLiteralName& EventName, const FFieldDesc* FieldDescs, uint32 FieldCount, uint32 Flags=0);
 };
 
diff --git a/Engine/Source/Runtime/TraceLog/Public/Trace/Detail/Trace.h b/Engine/Source/Runtime/TraceLog/Public/Trace/Detail/Trace.h
--- a/Engine/Source/Runtime/TraceLog/Public/Trace/Detail/Trace.h
+++ b/Engine/Source/Runtime/TraceLog/Public/Trace/Detail/Trace.h
@@ -8,6 +8,33 @@
 
 #include "CoreTypes.h"
 
+#define TRACE_PRIVATE_CHANNEL_DECLARE(LinkageType, ChannelName) \
+	LinkageType Trace::FChannel ChannelName;
+
+#define TRACE_PRIVATE_CHANNEL_IMPL(ChannelName) \
+	struct F##ChannelName##Registrator \
+	{ \
+		F##ChannelName##Registrator() \
+		{ \
+			Trace::FChannel::Register(ChannelName, PREPROCESSOR_TO_STRING(ChannelName)); \
+		} \
+	}; \
+	static F##ChannelName##Registrator ChannelName##Reg = F##ChannelName##Registrator();
+
+#define TRACE_PRIVATE_CHANNEL(ChannelName) \
+	TRACE_PRIVATE_CHANNEL_DECLARE(static, ChannelName) \
+	TRACE_PRIVATE_CHANNEL_IMPL(ChannelName)
+
+#define TRACE_PRIVATE_CHANNEL_EXTERN(ChannelName) \
+	TRACE_PRIVATE_CHANNEL_DECLARE(extern, ChannelName)
+
+#define TRACE_PRIVATE_CHANNEL_DEFINE(ChannelName) \
+	TRACE_PRIVATE_CHANNEL_DECLARE(, ChannelName) \
+	TRACE_PRIVATE_CHANNEL_IMPL(ChannelName)
+
+#define TRACE_PRIVATE_CHANNELEXPR_IS_ENABLED(ChannelsExpr) \
+	bool(ChannelsExpr)
+
 #define TRACE_PRIVATE_EVENT_DEFINE(LoggerName, EventName) \
 	Trace::FEventDef LoggerName##EventName##Event;
 
@@ -25,7 +52,6 @@
 		{ \
 			static const bool bOnceOnly = [] () \
 			{ \
-				const uint32 Always = Trace::FEventDef::Flag_Always; \
 				const uint32 Important = Trace::FEventDef::Flag_Important; \
 				uint32 Flags = decltype(EventProps_Private)::MaybeHasAux ? Trace::FEventDef::Flag_MaybeHasAux : 0; \
 				Flags |= (0, ##__VA_ARGS__); \
@@ -52,17 +78,19 @@
 		explicit operator bool () const { return true; } \
 	};
 
+#define TRACE_PRIVATE_EVENT_IS_IMPORTANT(LoggerName, EventName) \
+	( LoggerName##EventName##Event.bImportant )
+
 #define TRACE_PRIVATE_EVENT_IS_ENABLED(LoggerName, EventName) \
 	( \
 		(LoggerName##EventName##Event.bInitialized || (F##LoggerName##EventName##Fields::Initialize(), true)) \
-		&& (LoggerName##EventName##Event.Enabled.Test) \
 	)
 
 #define TRACE_PRIVATE_EVENT_SIZE(LoggerName, EventName) \
 	decltype(F##LoggerName##EventName##Fields::EventProps_Private)::Size
 
-#define TRACE_PRIVATE_LOG(LoggerName, EventName, ...) \
-	if (TRACE_PRIVATE_EVENT_IS_ENABLED(LoggerName, EventName)) \
+#define TRACE_PRIVATE_LOG(LoggerName, EventName, ChannelsExpr, ...) \
+	if ((TRACE_PRIVATE_CHANNELEXPR_IS_ENABLED(ChannelsExpr) || TRACE_PRIVATE_EVENT_IS_IMPORTANT(LoggerName, EventName)) && TRACE_PRIVATE_EVENT_IS_ENABLED(LoggerName, EventName)) \
 		if (const auto& __restrict EventName = (F##LoggerName##EventName##Fields&)LoggerName##EventName##Event) \
 			if (auto LogScope = Trace::FEventDef::FLogScope( \
 				LoggerName##EventName##Event.Uid, \
@@ -73,6 +101,15 @@
 
 #else
 
+#define TRACE_PRIVATE_CHANNEL(ChannelName)
+
+#define TRACE_PRIVATE_CHANNEL_EXTERN(ModuleApi, ChannelName)
+
+#define TRACE_PRIVATE_CHANNEL_DEFINE(ChannelName)
+
+#define TRACE_PRIVATE_CHANNELEXPR_IS_ENABLED(ChannelsExpr)\
+	false
+
 #define TRACE_PRIVATE_EVENT_DEFINE(LoggerName, EventName)
 
 #define TRACE_PRIVATE_EVENT_BEGIN(LoggerName, EventName, ...) \
@@ -101,7 +138,7 @@
 		const FTraceDisabled& Attachment; \
 	};
 
-#define TRACE_PRIVATE_EVENT_IS_ENABLED(LoggerName, EventName) \
+#define TRACE_PRIVATE_EVENT_IS_INITIALIZED(LoggerName, EventName) \
 	(false)
 
 #define TRACE_PRIVATE_LOG(LoggerName, EventName, ...) \
diff --git a/Engine/Source/Runtime/TraceLog/Public/Trace/Trace.h b/Engine/Source/Runtime/TraceLog/Public/Trace/Trace.h
--- a/Engine/Source/Runtime/TraceLog/Public/Trace/Trace.h
+++ b/Engine/Source/Runtime/TraceLog/Public/Trace/Trace.h
@@ -18,9 +18,12 @@
 namespace Trace
 {
 
+UE_TRACE_API bool	Initialize() UE_TRACE_IMPL(false);
 UE_TRACE_API bool	SendTo(const TCHAR* Host) UE_TRACE_IMPL(false);
 UE_TRACE_API bool	WriteTo(const TCHAR* Path) UE_TRACE_IMPL(false);
 UE_TRACE_API uint32 ToggleEvent(const TCHAR* Wildcard, bool bState) UE_TRACE_IMPL(0);
+UE_TRACE_API bool	ToggleChannel(const TCHAR* ChannelName, bool bEnabled) UE_TRACE_IMPL(false);
+UE_TRACE_API bool	ToggleChannel(struct FChannel& Channel, bool bEnabled) UE_TRACE_IMPL(false);
 
 } // namespace Trace
 
@@ -31,7 +34,14 @@ UE_TRACE_API uint32 ToggleEvent(const TCHAR* Wildcard, bool bState) UE_TRACE_IMP
 #define UE_TRACE_EVENT_FIELD(FieldType, FieldName)				TRACE_PRIVATE_EVENT_FIELD(FieldType, FieldName)
 #define UE_TRACE_EVENT_END()									TRACE_PRIVATE_EVENT_END()
 #define UE_TRACE_EVENT_IS_ENABLED(LoggerName, EventName)		TRACE_PRIVATE_EVENT_IS_ENABLED(LoggerName, EventName)
-#define UE_TRACE_LOG(LoggerName, EventName, ...)				TRACE_PRIVATE_LOG(LoggerName, EventName, ##__VA_ARGS__)
+#define UE_TRACE_LOG(LoggerName, EventName, ChannelsExpr, ...)	TRACE_PRIVATE_LOG(LoggerName, EventName, ChannelsExpr, ##__VA_ARGS__)
+
+////////////////////////////////////////////////////////////////////////////////
+#define UE_TRACE_CHANNEL(ChannelName)							TRACE_PRIVATE_CHANNEL(ChannelName)
+#define UE_TRACE_CHANNEL_EXTERN(ChannelName)					TRACE_PRIVATE_CHANNEL_EXTERN(ChannelName)
+#define UE_TRACE_CHANNEL_DEFINE(ChannelName)					TRACE_PRIVATE_CHANNEL_DEFINE(ChannelName)
+#define UE_TRACE_CHANNELEXPR_IS_ENABLED(ChannelsExpr)			TRACE_PRIVATE_CHANNELEXPR_IS_ENABLED(ChannelsExpr)
 
 #include "Detail/EventDef.inl"
 #include "Detail/Writer.inl"
+#include "Detail/Channel.inl"
