Fix duplicated function names.

diff --git a/Engine/Source/Runtime/TraceLog/Private/Trace/Channel.cpp b/Engine/Source/Runtime/TraceLog/Private/Trace/Channel.cpp
--- a/Engine/Source/Runtime/TraceLog/Private/Trace/Channel.cpp
+++ b/Engine/Source/Runtime/TraceLog/Private/Trace/Channel.cpp
@@ -16,7 +16,7 @@ static const size_t ChannelNameMaxLength = 64u;
 
 ///////////////////////////////////////////////////////////////////////////////
 template <int DestSize>
-static void ToAnsiCheap(ANSICHAR(&Dest)[DestSize], const WIDECHAR* Src)
+static void ChannelToAnsiCheap(ANSICHAR(&Dest)[DestSize], const WIDECHAR* Src)
 {
 	for (ANSICHAR& Out : Dest)
 	{
@@ -158,7 +158,7 @@ void FChannel::Register(FChannel& Channel, const ANSICHAR* ChannelName)
  bool FChannel::Toggle(const TCHAR* ChannelName, bool bEnabled)
  {
 	 ANSICHAR ChannelNameA[ChannelNameMaxLength];
-	 ToAnsiCheap(ChannelNameA, ChannelName);
+	 ChannelToAnsiCheap(ChannelNameA, ChannelName);
 
 	 return FChannel::Toggle(ChannelNameA, bEnabled);
  }
