Channel strcpy fixes.

diff --git a/Engine/Source/Runtime/TraceLog/Private/Trace/Channel.cpp b/Engine/Source/Runtime/TraceLog/Private/Trace/Channel.cpp
--- a/Engine/Source/Runtime/TraceLog/Private/Trace/Channel.cpp
+++ b/Engine/Source/Runtime/TraceLog/Private/Trace/Channel.cpp
@@ -1,4 +1,4 @@
-// Copyright 1998-2019 Epic Games, Inc. All Rights Reserved.
+// Copyright Epic Games, Inc. All Rights Reserved.
 #include "Trace/Detail/Channel.h"
 #include "Trace/Trace.h"
 #include "Trace/Detail/Atomic.h"
@@ -57,13 +57,15 @@ static uint32 GetChannelIdentifier(ANSICHAR(&Buffer) [ChannelNameMaxLength], con
 		const size_t Count = ChannelStr - ChannelName;
 		if (Count > 0 && strlen(ChannelName) > Count)
 		{
-			strncpy_s(Buffer, ChannelNameMaxLength, ChannelName, Count);
+			memcpy(&Buffer, ChannelName, Count);
 			Buffer[Count] = '\0';
 		}
 	}
 	else
 	{
-		strcpy_s(Buffer, ChannelNameMaxLength, ChannelName);
+		const size_t Count = strlen(ChannelName);
+		memcpy(&Buffer, ChannelName, Count);
+		Buffer[Count] = '\0';
 	}
 
 	// Convert to lower case
