Clamp string copy exceed buffer size.

diff --git a/Engine/Source/Developer/TraceServices/Private/Model/Channel.h b/Engine/Source/Developer/TraceServices/Private/Model/Channel.h
--- a/Engine/Source/Developer/TraceServices/Private/Model/Channel.h
+++ b/Engine/Source/Developer/TraceServices/Private/Model/Channel.h
@@ -1,4 +1,4 @@
-// Copyright 1998-2019 Epic Games, Inc. All Rights Reserved.
+// Copyright Epic Games, Inc. All Rights Reserved.
 
 #pragma once
 
diff --git a/Engine/Source/Runtime/TraceLog/Private/Trace/Channel.cpp b/Engine/Source/Runtime/TraceLog/Private/Trace/Channel.cpp
--- a/Engine/Source/Runtime/TraceLog/Private/Trace/Channel.cpp
+++ b/Engine/Source/Runtime/TraceLog/Private/Trace/Channel.cpp
@@ -14,6 +14,9 @@ namespace Trace
 static FChannel* volatile GHeadChannel; // = nullptr;
 static const size_t ChannelNameMaxLength = 64u;
 
+// General trace channel. Used by all built in events.
+static Trace::FTraceChannel TraceLogChannel;
+
 ///////////////////////////////////////////////////////////////////////////////
 template <int DestSize>
 static void ChannelToAnsiCheap(ANSICHAR(&Dest)[DestSize], const WIDECHAR* Src)
@@ -41,41 +44,32 @@ static uint32 ChannelGetHash(const ElementType* Input, int32 Length = -1)
 	return Result;
 }
 
-////////////////////////////////////////////////////////////////////////////////
-static uint32 ChannelGetHash(uint32 LoggerHash, uint32 NameHash)
-{
-	uint32 Parts[3] = { LoggerHash, NameHash, 0 };
-	return ChannelGetHash(Parts);
-}
-
 ///////////////////////////////////////////////////////////////////////////////
 static uint32 GetChannelIdentifier(ANSICHAR(&Buffer) [ChannelNameMaxLength], const ANSICHAR* ChannelName)
 {
 	// Strip "Channel" suffix if it exists
+	size_t CharsToCopy = 0;
 	if (const ANSICHAR* ChannelStr = strstr(ChannelName, "Channel"))
 	{
-		const size_t Count = ChannelStr - ChannelName;
-		if (Count > 0 && strlen(ChannelName) > Count)
-		{
-			memcpy(&Buffer, ChannelName, Count);
-			Buffer[Count] = '\0';
-		}
+		CharsToCopy = ChannelStr - ChannelName;
 	}
 	else
 	{
-		const size_t Count = strlen(ChannelName);
-		memcpy(&Buffer, ChannelName, Count);
-		Buffer[Count] = '\0';
+		CharsToCopy = strlen(ChannelName);
 	}
 
-	// Convert to lower case
-	const uint32 ChannelNameLen = strlen(Buffer);
-	for (size_t i = 0; i < ChannelNameLen; ++i)
+	// Turn name to lower case
+	if (CharsToCopy > 0 && strlen(ChannelName) >= CharsToCopy)
 	{
-		Buffer[i] = tolower(Buffer[i]);
+		const size_t CharsToCopySafe = CharsToCopy < (ChannelNameMaxLength - 1) ? CharsToCopy : ChannelNameMaxLength - 1;
+		for (size_t i = 0; i < CharsToCopySafe; i++)
+		{
+			Buffer[i] = tolower(*(ChannelName + i));
+		}
+		Buffer[CharsToCopySafe] = '\0';
 	}
 
-	return ChannelNameLen;
+	return strlen(Buffer);
 }
 
 ///////////////////////////////////////////////////////////////////////////////
@@ -109,7 +103,7 @@ void FChannel::Register(FChannel& Channel, const ANSICHAR* ChannelName)
 }
 
 ///////////////////////////////////////////////////////////////////////////////
- void FChannel::ToggleAll(bool bEnabled)
+void FChannel::ToggleAll(bool bEnabled)
 {
 	FChannel* Channel = Private::AtomicLoadAcquire(&GHeadChannel);
 	for (; Channel != nullptr; Channel = (FChannel*)(Channel->Handle))
@@ -118,9 +112,9 @@ void FChannel::Register(FChannel& Channel, const ANSICHAR* ChannelName)
 	}
 }
 
- ///////////////////////////////////////////////////////////////////////////////
- bool FChannel::Toggle(FChannel* Channel, bool bEnabled)
- {
+///////////////////////////////////////////////////////////////////////////////
+bool FChannel::Toggle(FChannel* Channel, bool bEnabled)
+{
 	UE_TRACE_EVENT_BEGIN($Trace, ChannelToggle, Important)
 		UE_TRACE_EVENT_FIELD(uint32, Id)
 		UE_TRACE_EVENT_FIELD(bool, IsEnabled)
@@ -135,11 +129,11 @@ void FChannel::Register(FChannel& Channel, const ANSICHAR* ChannelName)
 			<< ChannelToggle.IsEnabled(bEnabled);
 	}
 	return bWasEnabled;
- }
+}
 
- ///////////////////////////////////////////////////////////////////////////////
- bool FChannel::Toggle(const ANSICHAR* ChannelName, bool bEnabled)
- {
+///////////////////////////////////////////////////////////////////////////////
+bool FChannel::Toggle(const ANSICHAR* ChannelName, bool bEnabled)
+{
 	ANSICHAR ChannelIdentifier[ChannelNameMaxLength];
 	const uint32 ChannelNameLen = GetChannelIdentifier(ChannelIdentifier, ChannelName);
 	const uint32 ChannelNameHash = ChannelGetHash(ChannelIdentifier);
@@ -154,16 +148,16 @@ void FChannel::Register(FChannel& Channel, const ANSICHAR* ChannelName)
 	}
 
 	return false;
- }
+}
 
- ///////////////////////////////////////////////////////////////////////////////
- bool FChannel::Toggle(const TCHAR* ChannelName, bool bEnabled)
- {
-	 ANSICHAR ChannelNameA[ChannelNameMaxLength];
-	 ChannelToAnsiCheap(ChannelNameA, ChannelName);
+///////////////////////////////////////////////////////////////////////////////
+bool FChannel::Toggle(const TCHAR* ChannelName, bool bEnabled)
+{
+	ANSICHAR ChannelNameA[ChannelNameMaxLength];
+	ChannelToAnsiCheap(ChannelNameA, ChannelName);
 
-	 return FChannel::Toggle(ChannelNameA, bEnabled);
- }
+	return FChannel::Toggle(ChannelNameA, bEnabled);
+}
 
 }
 
diff --git a/Engine/Source/Runtime/TraceLog/Private/Trace/Writer.cpp b/Engine/Source/Runtime/TraceLog/Private/Trace/Writer.cpp
--- a/Engine/Source/Runtime/TraceLog/Private/Trace/Writer.cpp
+++ b/Engine/Source/Runtime/TraceLog/Private/Trace/Writer.cpp
@@ -935,7 +935,7 @@ static void Writer_InitializeControl()
 			ANSICHAR Channels[BufferSize] = {};
 			ANSICHAR* Ctx;
 			const bool bState = (ArgV[1][0] != '0');
-			strcpy_s(Channels, BufferSize, ArgV[0]);
+			FPlatformString::Strcpy(Channels, BufferSize, ArgV[0]);
 			ANSICHAR* Channel = FPlatformString::Strtok(Channels, ",", &Ctx);
 			while (Channel)
 			{
diff --git a/Engine/Source/Runtime/TraceLog/Public/Trace/Detail/Channel.inl b/Engine/Source/Runtime/TraceLog/Public/Trace/Detail/Channel.inl
--- a/Engine/Source/Runtime/TraceLog/Public/Trace/Detail/Channel.inl
+++ b/Engine/Source/Runtime/TraceLog/Public/Trace/Detail/Channel.inl
@@ -37,7 +37,6 @@ struct FTraceChannel : public FChannel
 
 }
 
-// General trace channel. Used by all built in events.
-static Trace::FTraceChannel TraceLogChannel;
+extern TRACELOG_API Trace::FTraceChannel TraceLogChannel;
 
 #endif //UE_TRACE_ENABLED
\ No newline at end of file
