Header fix.

#rnx

diff --git a/Engine/Source/Runtime/TraceLog/Private/Trace/Channel.cpp b/Engine/Source/Runtime/TraceLog/Private/Trace/Channel.cpp
--- a/Engine/Source/Runtime/TraceLog/Private/Trace/Channel.cpp
+++ b/Engine/Source/Runtime/TraceLog/Private/Trace/Channel.cpp
@@ -3,7 +3,8 @@
 #include "Trace/Trace.h"
 #include "Trace/Detail/Atomic.h"
 #include "Trace/Detail/Channel.h"
-#include "Misc/Parse.h"
+
+#include <ctype.h>
 
 #if UE_TRACE_ENABLED
 
