Added option to specify the TCP port to send traces to.

#rnx

diff --git a/Engine/Source/Runtime/TraceLog/Private/Trace/Trace.cpp b/Engine/Source/Runtime/TraceLog/Private/Trace/Trace.cpp
--- a/Engine/Source/Runtime/TraceLog/Private/Trace/Trace.cpp
+++ b/Engine/Source/Runtime/TraceLog/Private/Trace/Trace.cpp
@@ -13,7 +13,7 @@ namespace Private
 {
 
 ////////////////////////////////////////////////////////////////////////////////
-bool	Writer_SendTo(const ANSICHAR*);
+bool	Writer_SendTo(const ANSICHAR*, uint32);
 bool	Writer_WriteTo(const ANSICHAR*);
 uint32	Writer_EventToggle(const ANSICHAR*, bool);
 
@@ -43,11 +43,11 @@ bool Initialize()
 }
 
 ////////////////////////////////////////////////////////////////////////////////
-bool SendTo(const TCHAR* InHost)
+bool SendTo(const TCHAR* InHost, uint32 Port)
 {
 	char Host[32];
 	ToAnsiCheap(Host, InHost);
-	return Private::Writer_SendTo(Host);
+	return Private::Writer_SendTo(Host, Port);
 }
 
 ////////////////////////////////////////////////////////////////////////////////
diff --git a/Engine/Source/Runtime/TraceLog/Private/Trace/Writer.cpp b/Engine/Source/Runtime/TraceLog/Private/Trace/Writer.cpp
--- a/Engine/Source/Runtime/TraceLog/Private/Trace/Writer.cpp
+++ b/Engine/Source/Runtime/TraceLog/Private/Trace/Writer.cpp
@@ -670,7 +670,7 @@ struct FControlCommands
 };
 
 ////////////////////////////////////////////////////////////////////////////////
-bool	Writer_SendTo(const ANSICHAR*);
+bool	Writer_SendTo(const ANSICHAR*, uint32);
 bool	Writer_WriteTo(const ANSICHAR*);
 uint32	Writer_EventToggle(const ANSICHAR*, bool);
 
@@ -895,7 +895,7 @@ static void Writer_InitializeControl()
 		{
 			if (ArgC > 0)
 			{
-				Writer_SendTo(ArgV[0]);
+				Writer_SendTo(ArgV[0], 1980);
 			}
 		}
 	);
@@ -1061,11 +1061,11 @@ void Writer_Initialize()
 
 
 ////////////////////////////////////////////////////////////////////////////////
-bool Writer_SendTo(const ANSICHAR* Host)
+bool Writer_SendTo(const ANSICHAR* Host, uint32 Port)
 {
 	Writer_Initialize();
 
-	UPTRINT DataHandle = TcpSocketConnect(Host, 1980);
+	UPTRINT DataHandle = TcpSocketConnect(Host, Port);
 	if (!DataHandle)
 	{
 		return false;
diff --git a/Engine/Source/Runtime/TraceLog/Public/Trace/Trace.h b/Engine/Source/Runtime/TraceLog/Public/Trace/Trace.h
--- a/Engine/Source/Runtime/TraceLog/Public/Trace/Trace.h
+++ b/Engine/Source/Runtime/TraceLog/Public/Trace/Trace.h
@@ -19,7 +19,7 @@ namespace Trace
 {
 
 UE_TRACE_API bool	Initialize() UE_TRACE_IMPL(false);
-UE_TRACE_API bool	SendTo(const TCHAR* Host) UE_TRACE_IMPL(false);
+UE_TRACE_API bool	SendTo(const TCHAR* Host, uint32 Port=1980) UE_TRACE_IMPL(false);
 UE_TRACE_API bool	WriteTo(const TCHAR* Path) UE_TRACE_IMPL(false);
 UE_TRACE_API uint32 ToggleEvent(const TCHAR* Wildcard, bool bState) UE_TRACE_IMPL(0);
 UE_TRACE_API bool	ToggleChannel(const TCHAR* ChannelName, bool bEnabled) UE_TRACE_IMPL(false);
