Fix a linker issue when using TraceLogChannel with a UE_TRACE_LOG.

TraceLogChannel was inside the Trace namespace in Channel.cpp but in the global namespace in Channel.inl, since we don't want people to have to write UE_TRACE_LOG( ..., Trace::TraceLogChannel ... ) I moved it out of the namespace in the cpp.
It was also being declared as static in the cpp which conflicts with the extern keyword in the inline file, so I removed that as well.

#rnx

diff --git a/Engine/Source/Runtime/TraceLog/Private/Trace/Channel.cpp b/Engine/Source/Runtime/TraceLog/Private/Trace/Channel.cpp
--- a/Engine/Source/Runtime/TraceLog/Private/Trace/Channel.cpp
+++ b/Engine/Source/Runtime/TraceLog/Private/Trace/Channel.cpp
@@ -8,6 +8,9 @@
 
 #if UE_TRACE_ENABLED
 
+// General trace channel. Used by all built in events.
+Trace::FTraceChannel TraceLogChannel;
+
 namespace Trace
 {
 
@@ -15,9 +18,6 @@ namespace Trace
 static FChannel* volatile GHeadChannel; // = nullptr;
 static const size_t ChannelNameMaxLength = 64u;
 
-// General trace channel. Used by all built in events.
-static Trace::FTraceChannel TraceLogChannel;
-
 ///////////////////////////////////////////////////////////////////////////////
 template <int DestSize>
 static void ChannelToAnsiCheap(ANSICHAR(&Dest)[DestSize], const WIDECHAR* Src)
