Added missing null terminator.

#rnx

diff --git a/Engine/Source/Runtime/TraceLog/Private/Trace/Trace.cpp b/Engine/Source/Runtime/TraceLog/Private/Trace/Trace.cpp
--- a/Engine/Source/Runtime/TraceLog/Private/Trace/Trace.cpp
+++ b/Engine/Source/Runtime/TraceLog/Private/Trace/Trace.cpp
@@ -33,6 +33,7 @@ static void ToAnsiCheap(ANSICHAR (&Dest)[DestSize], const WIDECHAR* Src)
 			break;
 		}
 	}
+	Dest[DestSize - 1] = '\0';
 };
 
 ////////////////////////////////////////////////////////////////////////////////
