Remove use of "Always" from inactive code blocks.

diff --git a/Engine/Source/Runtime/Core/Private/Stats/StatsTrace.cpp b/Engine/Source/Runtime/Core/Private/Stats/StatsTrace.cpp
--- a/Engine/Source/Runtime/Core/Private/Stats/StatsTrace.cpp
+++ b/Engine/Source/Runtime/Core/Private/Stats/StatsTrace.cpp
@@ -12,14 +12,14 @@
 
 UE_TRACE_CHANNEL(StatsChannel)
 
-UE_TRACE_EVENT_BEGIN(Stats, Spec, Always)
+UE_TRACE_EVENT_BEGIN(Stats, Spec, Important)
 	UE_TRACE_EVENT_FIELD(uint32, Id)
 	UE_TRACE_EVENT_FIELD(bool, IsFloatingPoint)
 	UE_TRACE_EVENT_FIELD(bool, IsMemory)
 	UE_TRACE_EVENT_FIELD(bool, ShouldClearEveryFrame)
 UE_TRACE_EVENT_END()
 
-UE_TRACE_EVENT_BEGIN(Stats, EventBatch, Always)
+UE_TRACE_EVENT_BEGIN(Stats, EventBatch)
 	UE_TRACE_EVENT_FIELD(uint32, ThreadId)
 UE_TRACE_EVENT_END()
 
diff --git a/Engine/Source/Runtime/RHI/Private/GpuProfilerTrace.cpp b/Engine/Source/Runtime/RHI/Private/GpuProfilerTrace.cpp
--- a/Engine/Source/Runtime/RHI/Private/GpuProfilerTrace.cpp
+++ b/Engine/Source/Runtime/RHI/Private/GpuProfilerTrace.cpp
@@ -10,7 +10,7 @@
 
 UE_TRACE_CHANNEL(GpuProfilerChannel)
 
-UE_TRACE_EVENT_BEGIN(GpuProfiler, EventSpec, Always)
+UE_TRACE_EVENT_BEGIN(GpuProfiler, EventSpec, Important)
 	UE_TRACE_EVENT_FIELD(const void*, EventType)
 	UE_TRACE_EVENT_FIELD(uint16, NameLength)
 UE_TRACE_EVENT_END()
diff --git a/Engine/Source/Runtime/TraceLog/Private/Trace/Writer.cpp b/Engine/Source/Runtime/TraceLog/Private/Trace/Writer.cpp
--- a/Engine/Source/Runtime/TraceLog/Private/Trace/Writer.cpp
+++ b/Engine/Source/Runtime/TraceLog/Private/Trace/Writer.cpp
@@ -26,12 +26,12 @@ int32 Encode(const void*, int32, void*, int32);
 ////////////////////////////////////////////////////////////////////////////////
 #define TRACE_PRIVATE_PERF 0
 #if TRACE_PRIVATE_PERF
-UE_TRACE_EVENT_BEGIN($Trace, WorkerThread, Always)
+UE_TRACE_EVENT_BEGIN($Trace, WorkerThread)
 	UE_TRACE_EVENT_FIELD(uint32, Cycles)
 	UE_TRACE_EVENT_FIELD(uint32, BytesSent)
 UE_TRACE_EVENT_END()
 
-UE_TRACE_EVENT_BEGIN($Trace, Memory, Always)
+UE_TRACE_EVENT_BEGIN($Trace, Memory)
 	UE_TRACE_EVENT_FIELD(uint32, AllocSize)
 UE_TRACE_EVENT_END()
 #endif // TRACE_PRIVATE_PERF
