Removse event enable/disable.

diff --git a/Engine/Source/Developer/TraceAnalysis/Private/ControlClient.cpp b/Engine/Source/Developer/TraceAnalysis/Private/ControlClient.cpp
--- a/Engine/Source/Developer/TraceAnalysis/Private/ControlClient.cpp
+++ b/Engine/Source/Developer/TraceAnalysis/Private/ControlClient.cpp
@@ -121,12 +121,6 @@ void FControlClient::SendWriteTo(const TCHAR* Path)
     FormatAndSend(TEXT("WriteTo %s"), Path);
 }
 
-////////////////////////////////////////////////////////////////////////////////
-void FControlClient::SendToggleEvent(const TCHAR* EventSpec, bool bState)
-{
-    FormatAndSend(TEXT("ToggleEvent %s %d"), EventSpec, bState);
-}
-
 ////////////////////////////////////////////////////////////////////////////////
 void FControlClient::SendToggleChannel(const TCHAR* Channels, bool bState /*= true*/)
 {
diff --git a/Engine/Source/Developer/TraceAnalysis/Private/Recorder/Recorder.cpp b/Engine/Source/Developer/TraceAnalysis/Private/Recorder/Recorder.cpp
--- a/Engine/Source/Developer/TraceAnalysis/Private/Recorder/Recorder.cpp
+++ b/Engine/Source/Developer/TraceAnalysis/Private/Recorder/Recorder.cpp
@@ -33,7 +33,6 @@ private:
 	virtual void		StopRecording() override;
 	virtual uint32		GetSessionCount() const override;
 	virtual void		GetActiveSessions(TArray<FRecorderSessionInfo>& OutSessions) const override;
-	virtual bool		ToggleEvent(FRecorderSessionHandle RecordingHandle, const TCHAR* LoggerWildcard, bool bState) override;
 	virtual bool		ToggleChannels(FRecorderSessionHandle RecordingHandle, const TCHAR* Channels, bool bState) override;
 	FSession*			AcceptSession(FSocket& Socket);
 	void				CloseSession(FSession& Session);
@@ -328,29 +327,6 @@ void FRecorder::GetActiveSessions(TArray<FRecorderSessionInfo>& OutSessions) con
 	}
 }
 
-////////////////////////////////////////////////////////////////////////////////
-bool FRecorder::ToggleEvent(FRecorderSessionHandle SessionHandle, const TCHAR* LoggerWildcard, bool bState)
-{
-	TSharedPtr<FInternetAddr> ControlClientAddress;
-	{
-		FScopeLock Lock(&SessionsCS);
-		FSession** FindIt = Sessions.Find(SessionHandle);
-		if (!FindIt)
-		{
-			return false;
-		}
-		ControlClientAddress = (*FindIt)->ControlClientAddress;
-	}
-	FControlClient ControlClient;
-	if (!ControlClient.Connect(*ControlClientAddress))
-	{
-		return false;
-	}
-	ControlClient.SendToggleEvent(LoggerWildcard, bState);
-	ControlClient.Disconnect();
-	return true;
-}
-
 ////////////////////////////////////////////////////////////////////////////////
 bool FRecorder::ToggleChannels(FRecorderSessionHandle RecordingHandle, const TCHAR* Channels, bool bState)
 {
diff --git a/Engine/Source/Developer/TraceAnalysis/Public/Trace/ControlClient.h b/Engine/Source/Developer/TraceAnalysis/Public/Trace/ControlClient.h
--- a/Engine/Source/Developer/TraceAnalysis/Public/Trace/ControlClient.h
+++ b/Engine/Source/Developer/TraceAnalysis/Public/Trace/ControlClient.h
@@ -44,13 +44,6 @@ public:
 	 * @param Path File system path that the runtime should write events to. */
     void SendWriteTo(const TCHAR* Path);
 
-	/** Toggles events on and off. Events can be specified in two forms; either
-	 * by logger to address a logger's events as a whole (EventSpec="LoggerName")
-	 * or individually (EventSpec="LoggerName.EventName").
-	 * @param EventSpec Specification of the event to toggle.
-	 * @param bState True to enable the event and false to disable it. */
-    void SendToggleEvent(const TCHAR* EventSpec, bool bState=true);
-	
 	/**
 	 * Toggles one or more channels on and off. A channel or a comma separated list
 	 * of channels can be controlled.
diff --git a/Engine/Source/Developer/TraceAnalysis/Public/Trace/Recorder.h b/Engine/Source/Developer/TraceAnalysis/Public/Trace/Recorder.h
--- a/Engine/Source/Developer/TraceAnalysis/Public/Trace/Recorder.h
+++ b/Engine/Source/Developer/TraceAnalysis/Public/Trace/Recorder.h
@@ -47,13 +47,6 @@ public:
 	 * @param OutSessions TArray instance to receive the current session info. */
 	virtual void GetActiveSessions(TArray<FRecorderSessionInfo>& OutSessions) const = 0;
 
-	/** Toggles an event logger wildcard on or off for an active recording session.
-	 * @param RecordingHandle Identifier of the active session, from GetActiveSessions().
-	 * @param EventSpec Event specification (see FControlClient::ToggleEvent)
-	 * @param bState True to enable the event and false to disable it.
-	 * @return True if the command was successfully sent. */
-	virtual bool ToggleEvent(FRecorderSessionHandle RecordingHandle, const TCHAR* EventSpec, bool bState) = 0;
-
 	/**
 	 * Toggles one or more channels on and off. A channel or a comma separated list
 	 * of channels can be controlled.
diff --git a/Engine/Source/Developer/TraceServices/Private/SessionService.cpp b/Engine/Source/Developer/TraceServices/Private/SessionService.cpp
--- a/Engine/Source/Developer/TraceServices/Private/SessionService.cpp
+++ b/Engine/Source/Developer/TraceServices/Private/SessionService.cpp
@@ -193,58 +193,6 @@ bool FSessionService::GetSessionInfo(FSessionHandle SessionHandle, FSessionInfo&
 	return true;
 }
 
-void FSessionService::SetModuleEnabled(FSessionHandle SessionHandle, const FName& ModuleName, bool bState)
-{
-	FScopeLock Lock(&SessionsCS);
-	FSessionInfoInternal* FindIt = Sessions.Find(SessionHandle);
-	if (!FindIt)
-	{
-		return;
-	}
-	if (bState)
-	{
-		TArray<const TCHAR*> Loggers = ModuleService.GetModuleLoggers(ModuleName);
-		TSet<FString>& EnabledLoggers = FindIt->EnabledModuleLoggersMap.FindOrAdd(ModuleName);
-		for (const TCHAR* Logger : Loggers)
-		{
-			EnabledLoggers.Add(Logger);
-		}
-		if (FindIt->RecorderSessionHandle)
-		{
-			for (const FString& Logger : EnabledLoggers)
-			{
-				TraceRecorder->ToggleEvent(FindIt->RecorderSessionHandle, *Logger, true);
-			}
-		}
-	}
-	else
-	{
-		TSet<FString>* EnabledLoggers = FindIt->EnabledModuleLoggersMap.Find(ModuleName);
-		if (EnabledLoggers)
-		{
-			if (FindIt->RecorderSessionHandle)
-			{
-				for (const FString& Logger : *EnabledLoggers)
-				{
-					TraceRecorder->ToggleEvent(FindIt->RecorderSessionHandle, *Logger, false);
-				}
-			}
-			FindIt->EnabledModuleLoggersMap.Remove(ModuleName);
-		}
-	}
-}
-
-bool FSessionService::IsModuleEnabled(Trace::FSessionHandle SessionHandle, const FName& ModuleName) const
-{
-	FScopeLock Lock(&SessionsCS);
-	const FSessionInfoInternal* FindIt = Sessions.Find(SessionHandle);
-	if (!FindIt)
-	{
-		return false;
-	}
-	return FindIt->EnabledModuleLoggersMap.Contains(ModuleName);
-}
-
 bool FSessionService::ToggleChannels(Trace::FSessionHandle SessionHandle, const TCHAR* Channels, bool bState)
 {
 	FScopeLock Lock(&SessionsCS);
diff --git a/Engine/Source/Developer/TraceServices/Private/SessionServicePrivate.h b/Engine/Source/Developer/TraceServices/Private/SessionServicePrivate.h
--- a/Engine/Source/Developer/TraceServices/Private/SessionServicePrivate.h
+++ b/Engine/Source/Developer/TraceServices/Private/SessionServicePrivate.h
@@ -28,8 +28,6 @@ public:
 	virtual void GetAvailableSessions(TArray<Trace::FSessionHandle>& OutSessions) const override;
 	virtual void GetLiveSessions(TArray<Trace::FSessionHandle>& OutSessions) const override;
 	virtual bool GetSessionInfo(Trace::FSessionHandle SessionHandle, Trace::FSessionInfo& OutSessionInfo) const override;
-	virtual void SetModuleEnabled(Trace::FSessionHandle SessionHandle, const FName& ModuleName, bool bState) override;
-	virtual bool IsModuleEnabled(Trace::FSessionHandle SessionHandle, const FName& ModuleName) const override;
 	virtual bool ToggleChannels(Trace::FSessionHandle SessionHandle, const TCHAR* Channels, bool bState) override;
 	virtual bool ConnectSession(const TCHAR* ControlClientAddress) override;
 	virtual FSessionHandle GetSessionHandleByName(const FString& SessionName) override;
diff --git a/Engine/Source/Developer/TraceServices/Public/TraceServices/SessionService.h b/Engine/Source/Developer/TraceServices/Public/TraceServices/SessionService.h
--- a/Engine/Source/Developer/TraceServices/Public/TraceServices/SessionService.h
+++ b/Engine/Source/Developer/TraceServices/Public/TraceServices/SessionService.h
@@ -39,8 +39,6 @@ public:
 	virtual void GetAvailableSessions(TArray<FSessionHandle>& OutSessions) const = 0;
 	virtual void GetLiveSessions(TArray<FSessionHandle>& OutSessions) const = 0;
 	virtual bool GetSessionInfo(FSessionHandle SessionHandle, FSessionInfo& OutSessionInfo) const = 0;
-	virtual void SetModuleEnabled(Trace::FSessionHandle SessionHandle, const FName& ModuleName, bool bState) = 0;
-	virtual bool IsModuleEnabled(Trace::FSessionHandle SessionHandle, const FName& ModuleName) const = 0;
 	virtual bool ToggleChannels(Trace::FSessionHandle SessionHandle, const TCHAR* Channels, bool bState) = 0;
 	virtual bool ConnectSession(const TCHAR* ControlClientAddress) = 0;
 	
diff --git a/Engine/Source/Runtime/Core/Private/ProfilingDebugging/CountersTrace.cpp b/Engine/Source/Runtime/Core/Private/ProfilingDebugging/CountersTrace.cpp
--- a/Engine/Source/Runtime/Core/Private/ProfilingDebugging/CountersTrace.cpp
+++ b/Engine/Source/Runtime/Core/Private/ProfilingDebugging/CountersTrace.cpp
@@ -62,8 +62,6 @@ void FCountersTrace::Init(const TCHAR* CmdLine)
 {
 	if (FParse::Param(CmdLine, TEXT("counterstrace")))
 	{
-		UE_TRACE_EVENT_IS_ENABLED(Counters, SetValueInt);
-		UE_TRACE_EVENT_IS_ENABLED(Counters, SetValueFloat);
 		Trace::ToggleChannel(CountersChannel, true);
 	}
 }
diff --git a/Engine/Source/Runtime/Core/Private/ProfilingDebugging/CpuProfilerTrace.cpp b/Engine/Source/Runtime/Core/Private/ProfilingDebugging/CpuProfilerTrace.cpp
--- a/Engine/Source/Runtime/Core/Private/ProfilingDebugging/CpuProfilerTrace.cpp
+++ b/Engine/Source/Runtime/Core/Private/ProfilingDebugging/CpuProfilerTrace.cpp
@@ -227,7 +227,6 @@ void FCpuProfilerTrace::Init(const TCHAR* CmdLine)
 {
 	if (FParse::Param(CmdLine, TEXT("cpuprofilertrace")))
 	{
-		UE_TRACE_EVENT_IS_ENABLED(CpuProfiler, EventBatch);
 		Trace::ToggleChannel(CpuProfilerChannel, true);
 	}
 }
diff --git a/Engine/Source/Runtime/Core/Private/ProfilingDebugging/PlatformFileTrace.cpp b/Engine/Source/Runtime/Core/Private/ProfilingDebugging/PlatformFileTrace.cpp
--- a/Engine/Source/Runtime/Core/Private/ProfilingDebugging/PlatformFileTrace.cpp
+++ b/Engine/Source/Runtime/Core/Private/ProfilingDebugging/PlatformFileTrace.cpp
@@ -141,14 +141,6 @@ void FPlatformFileTrace::Init(const TCHAR* CmdLine)
 {
 	if (FParse::Param(CmdLine, TEXT("filetrace")))
 	{
-		UE_TRACE_EVENT_IS_ENABLED(PlatformFile, BeginOpen);
-		UE_TRACE_EVENT_IS_ENABLED(PlatformFile, EndOpen);
-		UE_TRACE_EVENT_IS_ENABLED(PlatformFile, BeginClose);
-		UE_TRACE_EVENT_IS_ENABLED(PlatformFile, EndClose);
-		UE_TRACE_EVENT_IS_ENABLED(PlatformFile, BeginRead);
-		UE_TRACE_EVENT_IS_ENABLED(PlatformFile, EndRead);
-		UE_TRACE_EVENT_IS_ENABLED(PlatformFile, BeginWrite);
-		UE_TRACE_EVENT_IS_ENABLED(PlatformFile, EndWrite);
 		Trace::ToggleChannel(FileChannel, true);
 	}
 }
diff --git a/Engine/Source/Runtime/Core/Private/Serialization/LoadTimeTrace.cpp b/Engine/Source/Runtime/Core/Private/Serialization/LoadTimeTrace.cpp
--- a/Engine/Source/Runtime/Core/Private/Serialization/LoadTimeTrace.cpp
+++ b/Engine/Source/Runtime/Core/Private/Serialization/LoadTimeTrace.cpp
@@ -37,11 +37,6 @@ void FLoadTimeProfilerTrace::FRequestGroupScope::OutputBegin()
 		<< BeginRequestGroup.Attachment(Attachment);
 }
 
-void FLoadTimeProfilerTrace::InitInternal()
-{
-	UE_TRACE_EVENT_IS_ENABLED(LoadTime, BeginRequestGroup);
-	UE_TRACE_EVENT_IS_ENABLED(LoadTime, EndRequestGroup);
-}
 
 #endif
 
diff --git a/Engine/Source/Runtime/Core/Public/Serialization/LoadTimeTrace.h b/Engine/Source/Runtime/Core/Public/Serialization/LoadTimeTrace.h
--- a/Engine/Source/Runtime/Core/Public/Serialization/LoadTimeTrace.h
+++ b/Engine/Source/Runtime/Core/Public/Serialization/LoadTimeTrace.h
@@ -40,8 +40,6 @@ struct FLoadTimeProfilerTrace
 		uint16 FormatArgsSize = 0;
 		uint8 FormatArgsBuffer[1024];
 	};
-
-	CORE_API static void InitInternal();
 };
 
 #define TRACE_LOADTIME_REQUEST_GROUP_SCOPE(Format, ...) \
diff --git a/Engine/Source/Runtime/CoreUObject/Private/Serialization/LoadTimeTrace.cpp b/Engine/Source/Runtime/CoreUObject/Private/Serialization/LoadTimeTrace.cpp
--- a/Engine/Source/Runtime/CoreUObject/Private/Serialization/LoadTimeTrace.cpp
+++ b/Engine/Source/Runtime/CoreUObject/Private/Serialization/LoadTimeTrace.cpp
@@ -115,26 +115,6 @@ UE_TRACE_EVENT_END()
 
 void FLoadTimeProfilerTracePrivate::Init()
 {
-	FLoadTimeProfilerTrace::InitInternal();
-	UE_TRACE_EVENT_IS_ENABLED(LoadTime, StartAsyncLoading);
-	UE_TRACE_EVENT_IS_ENABLED(LoadTime, SuspendAsyncLoading);
-	UE_TRACE_EVENT_IS_ENABLED(LoadTime, ResumeAsyncLoading);
-	UE_TRACE_EVENT_IS_ENABLED(LoadTime, PackageSummary);
-	UE_TRACE_EVENT_IS_ENABLED(LoadTime, BeginCreateExport);
-	UE_TRACE_EVENT_IS_ENABLED(LoadTime, EndCreateExport);
-	UE_TRACE_EVENT_IS_ENABLED(LoadTime, BeginSerializeExport);
-	UE_TRACE_EVENT_IS_ENABLED(LoadTime, EndSerializeExport);
-	UE_TRACE_EVENT_IS_ENABLED(LoadTime, BeginPostLoadExport);
-	UE_TRACE_EVENT_IS_ENABLED(LoadTime, EndPostLoadExport);
-	UE_TRACE_EVENT_IS_ENABLED(LoadTime, BeginRequest);
-	UE_TRACE_EVENT_IS_ENABLED(LoadTime, EndRequest);
-	UE_TRACE_EVENT_IS_ENABLED(LoadTime, NewAsyncPackage);
-	UE_TRACE_EVENT_IS_ENABLED(LoadTime, BeginLoadAsyncPackage);
-	UE_TRACE_EVENT_IS_ENABLED(LoadTime, EndLoadAsyncPackage);
-	UE_TRACE_EVENT_IS_ENABLED(LoadTime, DestroyAsyncPackage);
-	UE_TRACE_EVENT_IS_ENABLED(LoadTime, AsyncPackageRequestAssociation);
-	UE_TRACE_EVENT_IS_ENABLED(LoadTime, AsyncPackageImportDependency);
-	UE_TRACE_EVENT_IS_ENABLED(LoadTime, ClassInfo);
 	if (FParse::Param(FCommandLine::Get(), TEXT("loadtimetrace")))
 	{
 		Trace::ToggleChannel(LoadTimeChannel, true);
diff --git a/Engine/Source/Runtime/Engine/Private/ObjectTrace.cpp b/Engine/Source/Runtime/Engine/Private/ObjectTrace.cpp
--- a/Engine/Source/Runtime/Engine/Private/ObjectTrace.cpp
+++ b/Engine/Source/Runtime/Engine/Private/ObjectTrace.cpp
@@ -80,9 +80,6 @@ void FObjectTrace::Init()
 {
 	if (FParse::Param(FCommandLine::Get(), TEXT("objecttrace")))
 	{
-		UE_TRACE_EVENT_IS_ENABLED(Object, Class);
-		UE_TRACE_EVENT_IS_ENABLED(Object, Object);
-		UE_TRACE_EVENT_IS_ENABLED(Object, ObjectEvent);
 		Trace::ToggleChannel(ObjectChannel, true);
 	}
 }
diff --git a/Engine/Source/Runtime/TraceLog/Private/Trace/Trace.cpp b/Engine/Source/Runtime/TraceLog/Private/Trace/Trace.cpp
--- a/Engine/Source/Runtime/TraceLog/Private/Trace/Trace.cpp
+++ b/Engine/Source/Runtime/TraceLog/Private/Trace/Trace.cpp
@@ -15,7 +15,6 @@ namespace Private
 ////////////////////////////////////////////////////////////////////////////////
 bool	Writer_SendTo(const ANSICHAR*, uint32);
 bool	Writer_WriteTo(const ANSICHAR*);
-uint32	Writer_EventToggle(const ANSICHAR*, bool);
 
 } // namespace Private
 
@@ -59,15 +58,6 @@ bool WriteTo(const TCHAR* InPath)
 	return Private::Writer_WriteTo(Path);
 }
 
-////////////////////////////////////////////////////////////////////////////////
-uint32 ToggleEvent(const TCHAR* Wildcard, bool bState)
-{
-	ANSICHAR WildcardA[64];
-	ToAnsiCheap(WildcardA, Wildcard);
-
-	return Private::Writer_EventToggle(WildcardA, bState);
-}
-
 ////////////////////////////////////////////////////////////////////////////////
 bool ToggleChannel(const TCHAR* ChannelName, bool bEnabled)
 {
diff --git a/Engine/Source/Runtime/TraceLog/Private/Trace/Writer.cpp b/Engine/Source/Runtime/TraceLog/Private/Trace/Writer.cpp
--- a/Engine/Source/Runtime/TraceLog/Private/Trace/Writer.cpp
+++ b/Engine/Source/Runtime/TraceLog/Private/Trace/Writer.cpp
@@ -674,7 +674,6 @@ struct FControlCommands
 ////////////////////////////////////////////////////////////////////////////////
 bool	Writer_SendTo(const ANSICHAR*, uint32);
 bool	Writer_WriteTo(const ANSICHAR*);
-uint32	Writer_EventToggle(const ANSICHAR*, bool);
 
 ////////////////////////////////////////////////////////////////////////////////
 static FControlCommands	GControlCommands;
@@ -912,19 +911,6 @@ static void Writer_InitializeControl()
 		}
 	);
 
-	Writer_ControlAddCommand("ToggleEvent", nullptr,
-		[] (void*, uint32 ArgC, ANSICHAR const* const* ArgV)
-		{
-			if (ArgC < 1)
-			{
-				return;
-			}
-			const ANSICHAR* Wildcard = ArgV[0];
-			const ANSICHAR* State = (ArgC > 1) ? ArgV[1] : "";
-			Writer_EventToggle(Wildcard, State[0] != '0');
-		}
-	);
-
 	Writer_ControlAddCommand("ToggleChannels", nullptr, 
 		[] (void*, uint32 ArgC, ANSICHAR const* const* ArgV) 
 		{
@@ -1141,8 +1127,6 @@ void Writer_EventCreate(
 	if (Uid >= uint32(EKnownEventUids::Max))
 	{
 		Target->Uid = uint16(EKnownEventUids::Invalid);
-		Target->Enabled.bOptedIn = false;
-		Target->Enabled.Internal = 0;
 		Target->bInitialized = true;
 		return;
 	}
@@ -1154,7 +1138,6 @@ void Writer_EventCreate(
  	Target->Uid = uint16(Uid);
 	Target->LoggerHash = LoggerHash;
 	Target->Hash = Writer_EventGetHash(LoggerHash, NameHash);
-	Target->Enabled.bOptedIn = false;
 	Target->bInitialized = true;
 	Target->bImportant = Flags & FEventDef::Flag_Important;
 
@@ -1232,48 +1215,6 @@ void Writer_EventCreate(
 	}
 }
 
-////////////////////////////////////////////////////////////////////////////////
-uint32 Writer_EventToggle(const ANSICHAR* Wildcard, bool bState)
-{
-	Writer_Initialize();
-
-	uint32 ToggleCount = 0;
-
-	const ANSICHAR* Dot = FCStringAnsi::Strchr(Wildcard, '.');
-	if (Dot == nullptr)
-	{
-		uint32 LoggerHash = Writer_EventGetHash(Wildcard);
-
-		FEventDef* Event = AtomicLoadAcquire(&GHeadEvent);
-		for (; Event != nullptr; Event = (FEventDef*)(Event->Handle))
-		{
-			if (Event->LoggerHash == LoggerHash)
-			{
-				Event->Enabled.bOptedIn = bState;
-				++ToggleCount;
-			}
-		}
-
-		return ToggleCount;
-	}
-
-	uint32 LoggerHash = Writer_EventGetHash(Wildcard, int(Dot - Wildcard));
-	uint32 NameHash = Writer_EventGetHash(Dot + 1);
-	uint32 EventHash = Writer_EventGetHash(LoggerHash, NameHash);
-
-	FEventDef* Event = (FEventDef*)AtomicLoadAcquire(&Private::GHeadEvent);
-	for (; Event != nullptr; Event = (FEventDef*)(Event->Handle))
-	{
-		if (Event->Hash == EventHash)
-		{
-			Event->Enabled.bOptedIn = bState;
-			++ToggleCount;
-		}
-	}
-
-	return ToggleCount;
-}
-
 } // namespace Private
 } // namespace Trace
 
diff --git a/Engine/Source/Runtime/TraceLog/Public/Trace/Detail/EventDef.h b/Engine/Source/Runtime/TraceLog/Public/Trace/Detail/EventDef.h
--- a/Engine/Source/Runtime/TraceLog/Public/Trace/Detail/EventDef.h
+++ b/Engine/Source/Runtime/TraceLog/Public/Trace/Detail/EventDef.h
@@ -38,15 +38,6 @@ public:
 	uint32						LoggerHash;
 	uint32						Hash;
 	uint16						Uid;
-	union
-	{
-		struct
-		{
-			bool				bOptedIn;
-			uint8				Internal;
-		};
-		uint16					Test;
-	}							Enabled;
 	bool						bInitialized;
 	bool						bImportant;
 	TRACELOG_API static void	Create(FEventDef* Target, const FLiteralName& LoggerName, const FLiteralName& EventName, const FFieldDesc* FieldDescs, uint32 FieldCount, uint32 Flags=0);
diff --git a/Engine/Source/Runtime/TraceLog/Public/Trace/Detail/Trace.h b/Engine/Source/Runtime/TraceLog/Public/Trace/Detail/Trace.h
--- a/Engine/Source/Runtime/TraceLog/Public/Trace/Detail/Trace.h
+++ b/Engine/Source/Runtime/TraceLog/Public/Trace/Detail/Trace.h
@@ -81,16 +81,14 @@
 #define TRACE_PRIVATE_EVENT_IS_IMPORTANT(LoggerName, EventName) \
 	( LoggerName##EventName##Event.bImportant )
 
-#define TRACE_PRIVATE_EVENT_IS_ENABLED(LoggerName, EventName) \
-	( \
-		(LoggerName##EventName##Event.bInitialized || (F##LoggerName##EventName##Fields::Initialize(), true)) \
-	)
+#define TRACE_PRIVATE_EVENT_IS_INITIALIZED(LoggerName, EventName) \
+	( LoggerName##EventName##Event.bInitialized || (F##LoggerName##EventName##Fields::Initialize(), true))
 
 #define TRACE_PRIVATE_EVENT_SIZE(LoggerName, EventName) \
 	decltype(F##LoggerName##EventName##Fields::EventProps_Private)::Size
 
 #define TRACE_PRIVATE_LOG(LoggerName, EventName, ChannelsExpr, ...) \
-	if ((TRACE_PRIVATE_CHANNELEXPR_IS_ENABLED(ChannelsExpr) || TRACE_PRIVATE_EVENT_IS_IMPORTANT(LoggerName, EventName)) && TRACE_PRIVATE_EVENT_IS_ENABLED(LoggerName, EventName)) \
+	if (TRACE_PRIVATE_EVENT_IS_INITIALIZED(LoggerName, EventName) && (TRACE_PRIVATE_CHANNELEXPR_IS_ENABLED(ChannelsExpr) || TRACE_PRIVATE_EVENT_IS_IMPORTANT(LoggerName, EventName))) \
 		if (const auto& __restrict EventName = (F##LoggerName##EventName##Fields&)LoggerName##EventName##Event) \
 			if (auto LogScope = Trace::FEventDef::FLogScope( \
 				LoggerName##EventName##Event.Uid, \
diff --git a/Engine/Source/Runtime/TraceLog/Public/Trace/Trace.h b/Engine/Source/Runtime/TraceLog/Public/Trace/Trace.h
--- a/Engine/Source/Runtime/TraceLog/Public/Trace/Trace.h
+++ b/Engine/Source/Runtime/TraceLog/Public/Trace/Trace.h
@@ -21,7 +21,6 @@ namespace Trace
 UE_TRACE_API bool	Initialize() UE_TRACE_IMPL(false);
 UE_TRACE_API bool	SendTo(const TCHAR* Host, uint32 Port=1980) UE_TRACE_IMPL(false);
 UE_TRACE_API bool	WriteTo(const TCHAR* Path) UE_TRACE_IMPL(false);
-UE_TRACE_API uint32 ToggleEvent(const TCHAR* Wildcard, bool bState) UE_TRACE_IMPL(0);
 UE_TRACE_API bool	ToggleChannel(const TCHAR* ChannelName, bool bEnabled) UE_TRACE_IMPL(false);
 UE_TRACE_API bool	ToggleChannel(struct FChannel& Channel, bool bEnabled) UE_TRACE_IMPL(false);
 
@@ -33,7 +32,6 @@ UE_TRACE_API bool	ToggleChannel(struct FChannel& Channel, bool bEnabled) UE_TRAC
 #define UE_TRACE_EVENT_BEGIN_EXTERN(LoggerName, EventName, ...)	TRACE_PRIVATE_EVENT_BEGIN_EXTERN(LoggerName, EventName, ##__VA_ARGS__)
 #define UE_TRACE_EVENT_FIELD(FieldType, FieldName)				TRACE_PRIVATE_EVENT_FIELD(FieldType, FieldName)
 #define UE_TRACE_EVENT_END()									TRACE_PRIVATE_EVENT_END()
-#define UE_TRACE_EVENT_IS_ENABLED(LoggerName, EventName)		TRACE_PRIVATE_EVENT_IS_ENABLED(LoggerName, EventName)
 #define UE_TRACE_LOG(LoggerName, EventName, ChannelsExpr, ...)	TRACE_PRIVATE_LOG(LoggerName, EventName, ChannelsExpr, ##__VA_ARGS__)
 
 ////////////////////////////////////////////////////////////////////////////////
