Removed event hashes and the linked list.

#rnx

diff --git a/Engine/Source/Runtime/TraceLog/Private/Trace/Writer.cpp b/Engine/Source/Runtime/TraceLog/Private/Trace/Writer.cpp
--- a/Engine/Source/Runtime/TraceLog/Private/Trace/Writer.cpp
+++ b/Engine/Source/Runtime/TraceLog/Private/Trace/Writer.cpp
@@ -1078,27 +1078,6 @@ bool Writer_WriteTo(const ANSICHAR* Path)
 
 ////////////////////////////////////////////////////////////////////////////////
 static UPTRINT volatile		GEventUidCounter;	// = 0;
-static FEventDef* volatile	GHeadEvent;			// = nullptr;
-
-////////////////////////////////////////////////////////////////////////////////
-template <typename ElementType>
-static uint32 Writer_EventGetHash(const ElementType* Input, int32 Length=-1)
-{
-	uint32 Result = 0x811c9dc5;
-	for (; *Input && Length; ++Input, --Length)
-	{
-		Result ^= *Input;
-		Result *= 0x01000193;
-	}
-	return Result;
-}
-
-////////////////////////////////////////////////////////////////////////////////
-static uint32 Writer_EventGetHash(uint32 LoggerHash, uint32 NameHash)
-{
-	uint32 Parts[3] = { LoggerHash, NameHash, 0 };
-	return Writer_EventGetHash(Parts);
-}
 
 ////////////////////////////////////////////////////////////////////////////////
 void Writer_EventCreate(
@@ -1131,13 +1110,8 @@ void Writer_EventCreate(
 		return;
 	}
 
-	uint32 LoggerHash = Writer_EventGetHash(LoggerName.Ptr);
-	uint32 NameHash = Writer_EventGetHash(EventName.Ptr);
-
 	// Fill out the target event's properties
  	Target->Uid = uint16(Uid);
-	Target->LoggerHash = LoggerHash;
-	Target->Hash = Writer_EventGetHash(LoggerHash, NameHash);
 	Target->bInitialized = true;
 	Target->bImportant = Flags & FEventDef::Flag_Important;
 
@@ -1202,17 +1176,6 @@ void Writer_EventCreate(
 	}
 
 	Writer_EndLog(LogInstance);
-
-	// Add this new event into the list so we can look them up later.
-	for (;; Private::PlatformYield())
-	{
-		FEventDef* HeadEvent = AtomicLoadRelaxed(&GHeadEvent);
-		Target->Handle = HeadEvent;
-		if (AtomicCompareExchangeRelease(&GHeadEvent, Target, HeadEvent))
-		{
-			break;
-		}
-	}
 }
 
 } // namespace Private
diff --git a/Engine/Source/Runtime/TraceLog/Public/Trace/Detail/EventDef.h b/Engine/Source/Runtime/TraceLog/Public/Trace/Detail/EventDef.h
--- a/Engine/Source/Runtime/TraceLog/Public/Trace/Detail/EventDef.h
+++ b/Engine/Source/Runtime/TraceLog/Public/Trace/Detail/EventDef.h
@@ -34,9 +34,6 @@ public:
 		constexpr explicit		operator bool () const { return true; }
 	};
 
-	void*						Handle;
-	uint32						LoggerHash;
-	uint32						Hash;
 	uint16						Uid;
 	bool						bInitialized;
 	bool						bImportant;
