Only channels should dictate if an event is sent or not.

#rnx

diff --git a/Engine/Source/Runtime/TraceLog/Public/Trace/Detail/Trace.h b/Engine/Source/Runtime/TraceLog/Public/Trace/Detail/Trace.h
--- a/Engine/Source/Runtime/TraceLog/Public/Trace/Detail/Trace.h
+++ b/Engine/Source/Runtime/TraceLog/Public/Trace/Detail/Trace.h
@@ -78,17 +78,14 @@
 		explicit operator bool () const { return true; } \
 	};
 
-#define TRACE_PRIVATE_EVENT_IS_IMPORTANT(LoggerName, EventName) \
-	( LoggerName##EventName##Event.bImportant )
-
 #define TRACE_PRIVATE_EVENT_IS_INITIALIZED(LoggerName, EventName) \
-	( LoggerName##EventName##Event.bInitialized || (F##LoggerName##EventName##Fields::Initialize(), true))
+	(LoggerName##EventName##Event.bInitialized || (F##LoggerName##EventName##Fields::Initialize(), true))
 
 #define TRACE_PRIVATE_EVENT_SIZE(LoggerName, EventName) \
 	decltype(F##LoggerName##EventName##Fields::EventProps_Private)::Size
 
 #define TRACE_PRIVATE_LOG(LoggerName, EventName, ChannelsExpr, ...) \
-	if (TRACE_PRIVATE_EVENT_IS_INITIALIZED(LoggerName, EventName) && (TRACE_PRIVATE_CHANNELEXPR_IS_ENABLED(ChannelsExpr) || TRACE_PRIVATE_EVENT_IS_IMPORTANT(LoggerName, EventName))) \
+	if (TRACE_PRIVATE_EVENT_IS_INITIALIZED(LoggerName, EventName) && TRACE_PRIVATE_CHANNELEXPR_IS_ENABLED(ChannelsExpr)) \
 		if (const auto& __restrict EventName = (F##LoggerName##EventName##Fields&)LoggerName##EventName##Event) \
 			if (auto LogScope = Trace::FEventDef::FLogScope( \
 				LoggerName##EventName##Event.Uid, \
@@ -105,7 +102,7 @@
 
 #define TRACE_PRIVATE_CHANNEL_DEFINE(ChannelName)
 
-#define TRACE_PRIVATE_CHANNELEXPR_IS_ENABLED(ChannelsExpr)\
+#define TRACE_PRIVATE_CHANNELEXPR_IS_ENABLED(ChannelsExpr) \
 	false
 
 #define TRACE_PRIVATE_EVENT_DEFINE(LoggerName, EventName)
