Only initialise a trace event if we are going to be sending it.

#rnx

diff --git a/Engine/Source/Runtime/TraceLog/Public/Trace/Detail/Trace.h b/Engine/Source/Runtime/TraceLog/Public/Trace/Detail/Trace.h
--- a/Engine/Source/Runtime/TraceLog/Public/Trace/Detail/Trace.h
+++ b/Engine/Source/Runtime/TraceLog/Public/Trace/Detail/Trace.h
@@ -85,7 +85,7 @@
 	decltype(F##LoggerName##EventName##Fields::EventProps_Private)::Size
 
 #define TRACE_PRIVATE_LOG(LoggerName, EventName, ChannelsExpr, ...) \
-	if (TRACE_PRIVATE_EVENT_IS_INITIALIZED(LoggerName, EventName) && TRACE_PRIVATE_CHANNELEXPR_IS_ENABLED(ChannelsExpr)) \
+	if (TRACE_PRIVATE_CHANNELEXPR_IS_ENABLED(ChannelsExpr) && TRACE_PRIVATE_EVENT_IS_INITIALIZED(LoggerName, EventName)) \
 		if (const auto& __restrict EventName = (F##LoggerName##EventName##Fields&)LoggerName##EventName##Event) \
 			if (auto LogScope = Trace::FEventDef::FLogScope( \
 				LoggerName##EventName##Event.Uid, \
@@ -133,9 +133,6 @@
 		const FTraceDisabled& Attachment; \
 	};
 
-#define TRACE_PRIVATE_EVENT_IS_INITIALIZED(LoggerName, EventName) \
-	(false)
-
 #define TRACE_PRIVATE_LOG(LoggerName, EventName, ...) \
 	if (const auto& EventName = *(F##LoggerName##EventName##Dummy*)1) \
 		EventName
