Compile-time access to a trace event's flags.

#rnx

diff --git a/Engine/Source/Runtime/TraceLog/Public/Trace/Detail/Trace.h b/Engine/Source/Runtime/TraceLog/Public/Trace/Detail/Trace.h
--- a/Engine/Source/Runtime/TraceLog/Public/Trace/Detail/Trace.h
+++ b/Engine/Source/Runtime/TraceLog/Public/Trace/Detail/Trace.h
@@ -48,19 +48,21 @@
 	LinkageType TRACE_PRIVATE_EVENT_DEFINE(LoggerName, EventName) \
 	struct F##LoggerName##EventName##Fields \
 	{ \
+		enum \
+		{ \
+			Important			= Trace::FEventDef::Flag_Important, \
+			PartialEventFlags	= (0, ##__VA_ARGS__), \
+		}; \
 		static void FORCENOINLINE Initialize() \
 		{ \
 			static const bool bOnceOnly = [] () \
 			{ \
-				const uint32 Important = Trace::FEventDef::Flag_Important; \
-				uint32 Flags = decltype(EventProps_Private)::MaybeHasAux ? Trace::FEventDef::Flag_MaybeHasAux : 0; \
-				Flags |= (0, ##__VA_ARGS__); \
 				F##LoggerName##EventName##Fields Fields; \
 				const auto* Descs = (Trace::FFieldDesc*)(&Fields); \
 				uint32 DescCount = uint32(sizeof(Fields) / sizeof(*Descs)); \
 				const auto& LoggerLiteral = Trace::FLiteralName(#LoggerName); \
 				const auto& EventLiteral = Trace::FLiteralName(#EventName); \
-				Trace::FEventDef::Create(&LoggerName##EventName##Event, LoggerLiteral, EventLiteral, Descs, DescCount, Flags); \
+				Trace::FEventDef::Create(&LoggerName##EventName##Event, LoggerLiteral, EventLiteral, Descs, DescCount, EventFlags); \
 				return true; \
 			}(); \
 		} \
@@ -76,6 +78,7 @@
 		Trace::EventProps> const EventProps_Private = {}; \
 		Trace::TField<0, decltype(EventProps_Private)::Size, Trace::Attachment> const Attachment = {}; \
 		explicit operator bool () const { return true; } \
+		enum { EventFlags = PartialEventFlags|(decltype(EventProps_Private)::MaybeHasAux ? Trace::FEventDef::Flag_MaybeHasAux : 0), }; \
 	};
 
 #define TRACE_PRIVATE_EVENT_IS_INITIALIZED(LoggerName, EventName) \
