Always send important events for now otherwise we miss channel announcements.

#rnx

diff --git a/Engine/Source/Runtime/TraceLog/Public/Trace/Detail/Trace.h b/Engine/Source/Runtime/TraceLog/Public/Trace/Detail/Trace.h
--- a/Engine/Source/Runtime/TraceLog/Public/Trace/Detail/Trace.h
+++ b/Engine/Source/Runtime/TraceLog/Public/Trace/Detail/Trace.h
@@ -84,11 +84,15 @@
 #define TRACE_PRIVATE_EVENT_IS_INITIALIZED(LoggerName, EventName) \
 	(LoggerName##EventName##Event.bInitialized || (F##LoggerName##EventName##Fields::Initialize(), true))
 
+#define TRACE_PRIVATE_EVENT_IS_IMPORTANT(LoggerName, EventName) \
+	(F##LoggerName##EventName##Fields::EventFlags & F##LoggerName##EventName##Fields::Important)
+
 #define TRACE_PRIVATE_EVENT_SIZE(LoggerName, EventName) \
 	decltype(F##LoggerName##EventName##Fields::EventProps_Private)::Size
 
 #define TRACE_PRIVATE_LOG(LoggerName, EventName, ChannelsExpr, ...) \
-	if (TRACE_PRIVATE_CHANNELEXPR_IS_ENABLED(ChannelsExpr) && TRACE_PRIVATE_EVENT_IS_INITIALIZED(LoggerName, EventName)) \
+	if ((TRACE_PRIVATE_CHANNELEXPR_IS_ENABLED(ChannelsExpr) || TRACE_PRIVATE_EVENT_IS_IMPORTANT(LoggerName, EventName)) \
+		&& TRACE_PRIVATE_EVENT_IS_INITIALIZED(LoggerName, EventName)) \
 		if (const auto& __restrict EventName = (F##LoggerName##EventName##Fields&)LoggerName##EventName##Event) \
 			if (auto LogScope = Trace::FEventDef::FLogScope( \
 				LoggerName##EventName##Event.Uid, \
