If an the header for a blob of auxilary (array) data ended up in the overflow a commit could get missed causing buffers of trace data to stop draining.

#rnx

diff --git a/Engine/Source/Runtime/TraceLog/Public/Trace/Detail/Field.h b/Engine/Source/Runtime/TraceLog/Public/Trace/Detail/Field.h
--- a/Engine/Source/Runtime/TraceLog/Public/Trace/Detail/Field.h
+++ b/Engine/Source/Runtime/TraceLog/Public/Trace/Detail/Field.h
@@ -154,13 +154,18 @@ private:
 		Header->Size = Size << 8;
 		Header->FieldIndex = uint8(0x80 | (Index & int(EIndexPack::FieldCountMask)));
 
-		bool bCommit = ((uint8*)Header == Buffer->Committed);
+		bool bCommit = ((uint8*)Header + bMaybeHasAux == Buffer->Committed);
 
 		// Array data
 		while (true)
 		{
 			if (Buffer->Cursor >= (uint8*)Buffer)
 			{
+				if (bCommit)
+				{
+					AtomicStoreRelease(&(uint8* volatile&)(Buffer->Committed), Buffer->Cursor);
+				}
+
 				Buffer = Writer_NextBuffer(0);
 				bCommit = true;
 			}
@@ -170,11 +175,6 @@ private:
 			memcpy(Buffer->Cursor, Data, SegmentSize);
 			Buffer->Cursor += SegmentSize;
 
-			if (bCommit)
-			{
-				AtomicStoreRelease(&(uint8* volatile&)(Buffer->Committed), Buffer->Cursor);
-			}
-
 			Size -= SegmentSize;
 			if (Size <= 0)
 			{
@@ -187,6 +187,11 @@ private:
 		// The auxilary data null terminator.
 		Buffer->Cursor[0] = 0;
 		Buffer->Cursor++;
+
+		if (bCommit)
+		{
+			AtomicStoreRelease(&(uint8* volatile&)(Buffer->Committed), Buffer->Cursor);
+		}
 	}
 };
 
