Use increment instead of a CAS loop.

#rnx

diff --git a/Engine/Source/Runtime/TraceLog/Private/Trace/Writer.cpp b/Engine/Source/Runtime/TraceLog/Private/Trace/Writer.cpp
--- a/Engine/Source/Runtime/TraceLog/Private/Trace/Writer.cpp
+++ b/Engine/Source/Runtime/TraceLog/Private/Trace/Writer.cpp
@@ -67,7 +67,7 @@ static bool GInitialized = false;
 ////////////////////////////////////////////////////////////////////////////////
 thread_local					FWriteTlsContext TlsContext;
 uint8							FWriteTlsContext::DefaultBuffer[];	// = {}
-UPTRINT volatile				FWriteTlsContext::ThreadIdCounter;	// = 0;
+uint32 volatile					FWriteTlsContext::ThreadIdCounter;	// = 0;
 TRACELOG_API uint32 volatile	GLogSerial;							// = 0;
 
 ////////////////////////////////////////////////////////////////////////////////
@@ -85,16 +85,7 @@ FWriteTlsContext::FWriteTlsContext()
 	Buffer = Target;
 
 	// Assign a new id to this thread.
-	for (;; Private::PlatformYield())
-	{
-		UPTRINT CandidateId = UPTRINT(AtomicLoadRelaxed((void* volatile*)&ThreadIdCounter));
-		UPTRINT NextId = CandidateId + 1;
-		if (AtomicCompareExchangeRelaxed(&ThreadIdCounter, NextId, CandidateId))
-		{
-			ThreadId = uint32(CandidateId);
-			break;
-		}
-	}
+	ThreadId = AtomicIncrementRelaxed(&ThreadIdCounter);
 }
 
 ////////////////////////////////////////////////////////////////////////////////
@@ -1077,7 +1068,7 @@ bool Writer_WriteTo(const ANSICHAR* Path)
 
 
 ////////////////////////////////////////////////////////////////////////////////
-static UPTRINT volatile		GEventUidCounter;	// = 0;
+static uint32 volatile GEventUidCounter; // = 0;
 
 ////////////////////////////////////////////////////////////////////////////////
 void Writer_EventCreate(
@@ -1091,18 +1082,7 @@ void Writer_EventCreate(
 	Writer_Initialize();
 
 	// Assign a unique ID for this event
-	uint32 Uid;
-	for (;; Private::PlatformYield())
-	{
-		UPTRINT CurrentUid = UPTRINT(AtomicLoadRelaxed((void* volatile*)&GEventUidCounter));
-		UPTRINT NextUid = CurrentUid + 1;
-		if (AtomicCompareExchangeRelaxed(&GEventUidCounter, NextUid, CurrentUid))
-		{
-			Uid = uint32(CurrentUid) + uint32(EKnownEventUids::User);
-			break;
-		}
-	}
-
+	uint32 Uid = AtomicIncrementRelaxed(&GEventUidCounter) + uint32(EKnownEventUids::User);
 	if (Uid >= uint32(EKnownEventUids::Max))
 	{
 		Target->Uid = uint16(EKnownEventUids::Invalid);
diff --git a/Engine/Source/Runtime/TraceLog/Public/Trace/Detail/Writer.inl b/Engine/Source/Runtime/TraceLog/Public/Trace/Detail/Writer.inl
--- a/Engine/Source/Runtime/TraceLog/Public/Trace/Detail/Writer.inl
+++ b/Engine/Source/Runtime/TraceLog/Public/Trace/Detail/Writer.inl
@@ -41,7 +41,7 @@ private:
 	FWriteBuffer*			Buffer;
 	uint32					ThreadId;
 	static uint8			DefaultBuffer[sizeof(FWriteBuffer)];
-	static UPTRINT volatile	ThreadIdCounter;
+	static uint32 volatile	ThreadIdCounter;
 };
 
 
