Extracted the Send() lambda into its own function to aid ConsomeEvents' complexity.

#rnx

diff --git a/Engine/Source/Runtime/TraceLog/Private/Trace/Writer.cpp b/Engine/Source/Runtime/TraceLog/Private/Trace/Writer.cpp
--- a/Engine/Source/Runtime/TraceLog/Private/Trace/Writer.cpp
+++ b/Engine/Source/Runtime/TraceLog/Private/Trace/Writer.cpp
@@ -28,6 +28,7 @@ int32 Encode(const void*, int32, void*, int32);
 #if TRACE_PRIVATE_PERF
 UE_TRACE_EVENT_BEGIN($Trace, WorkerThread)
 	UE_TRACE_EVENT_FIELD(uint32, Cycles)
+	UE_TRACE_EVENT_FIELD(uint32, BytesReaped)
 	UE_TRACE_EVENT_FIELD(uint32, BytesSent)
 UE_TRACE_EVENT_END()
 
@@ -61,6 +62,7 @@ void Writer_InitializeTiming()
 }
 
 
+
 ////////////////////////////////////////////////////////////////////////////////
 static bool GInitialized = false;
 
@@ -374,36 +376,9 @@ UPTRINT					GPendingDataHandle;	// = 0
 static FWriteBuffer*	GActiveBufferList;	// = nullptr;
 
 ////////////////////////////////////////////////////////////////////////////////
-static void Writer_ConsumeEvents()
+static uint32 Writer_SendData(uint32 ThreadId, uint8* __restrict Data, uint32 Size)
 {
-	uint64 StartTsc = TimeGetTimestamp();
-	uint32 BytesSent = 0;
-
-	// Claim ownership of the latest chain of sent events.
-	FWriteBuffer* __restrict NextBufferList;
-	for (;; Private::PlatformYield())
-	{
-		NextBufferList = AtomicLoadRelaxed(&GNextBufferList);
-		if (AtomicCompareExchangeAcquire(&GNextBufferList, (FWriteBuffer*)nullptr, NextBufferList))
-		{
-			break;
-		}
-	}
-
-	struct FRetireList
-	{
-		FWriteBuffer* __restrict Head = nullptr;
-		FWriteBuffer* __restrict Tail = nullptr;
-
-		void Insert(FWriteBuffer* __restrict Buffer)
-		{
-			Buffer->Next = Head;
-			Head = Buffer;
-			Tail = (Tail != nullptr) ? Tail : Head;
-		}
-	};
-
-	auto SendInner = [&] (uint8* __restrict Data, uint32 Size)
+	auto SendInner = [] (uint8* __restrict Data, uint32 Size)
 	{
 		if (GDataState == EDataState::Sending)
 		{
@@ -430,50 +405,78 @@ static void Writer_ConsumeEvents()
 		}
 	};
 
-	auto Send = [&] (FWriteBuffer* Buffer, uint8* __restrict Data, uint32 Size)
+	struct FPacketBase
 	{
-		BytesSent += Size;
+		uint16 PacketSize;
+		uint16 ThreadId;
+	};
 
-		struct FPacketBase
-		{
-			uint16 PacketSize;
-			uint16 ThreadId;
-		};
+	// Smaller buffers usually aren't redundant enough to benefit from being
+	// compressed. They often end up being larger.
+	if (Size <= 384)
+	{
+		static_assert(sizeof(FPacketBase) == sizeof(uint32), "");
+		Data -= sizeof(FPacketBase);
+		Size += sizeof(FPacketBase);
+		auto* Packet = (FPacketBase*)Data;
+		Packet->ThreadId = uint16(ThreadId & 0x7fff);
+		Packet->PacketSize = uint16(Size);
+
+		SendInner(Data, Size);
+		return Size;
+	}
 
-		// Smaller buffers usually aren't redundant enough to benefit from being
-		// compressed. They often end up being larger.
-		if (Size > 384)
-		{
-			struct FPacketEncoded
-				: public FPacketBase
-			{
-				uint16	DecodedSize;
-			};
+	struct FPacketEncoded
+		: public FPacketBase
+	{
+		uint16	DecodedSize;
+	};
 
-			struct FPacket
-				: public FPacketEncoded
-			{
-				uint8 Data[GPoolBlockSize + 64];
-			};
+	struct FPacket
+		: public FPacketEncoded
+	{
+		uint8 Data[GPoolBlockSize + 64];
+	};
 
-			FPacket Packet;
-			Packet.ThreadId = 0x8000 | uint16(Buffer->ThreadId & 0x7fff);
-			Packet.DecodedSize = uint16(Size);
-			Packet.PacketSize = Encode(Data, Packet.DecodedSize, Packet.Data, sizeof(Packet.Data));
-			Packet.PacketSize += sizeof(FPacketEncoded);
+	FPacket Packet;
+	Packet.ThreadId = 0x8000 | uint16(ThreadId & 0x7fff);
+	Packet.DecodedSize = uint16(Size);
+	Packet.PacketSize = Encode(Data, Packet.DecodedSize, Packet.Data, sizeof(Packet.Data));
+	Packet.PacketSize += sizeof(FPacketEncoded);
+
+	SendInner((uint8*)&Packet, Packet.PacketSize);
+	return Packet.PacketSize;
+}
 
-			SendInner((uint8*)&Packet, Packet.PacketSize);
+////////////////////////////////////////////////////////////////////////////////
+static void Writer_ConsumeEvents()
+{
+#if TRACE_PRIVATE_PERF
+	uint64 StartTsc = TimeGetTimestamp();
+	uint32 BytesSent = 0;
+#endif
+
+	// Claim ownership of the latest chain of sent events.
+	FWriteBuffer* __restrict NextBufferList;
+	for (;; Private::PlatformYield())
+	{
+		NextBufferList = AtomicLoadRelaxed(&GNextBufferList);
+		if (AtomicCompareExchangeAcquire(&GNextBufferList, (FWriteBuffer*)nullptr, NextBufferList))
+		{
+			break;
 		}
-		else
+	}
+
+	struct FRetireList
+	{
+		FWriteBuffer* __restrict Head = nullptr;
+		FWriteBuffer* __restrict Tail = nullptr;
+
+		void Insert(FWriteBuffer* __restrict Buffer)
 		{
-			static_assert(sizeof(FPacketBase) == sizeof(uint32), "");
-			Data -= sizeof(FPacketBase);
-			Size += sizeof(FPacketBase);
-			auto* Packet = (FPacketBase*)Data;
-			Packet->ThreadId = uint16(Buffer->ThreadId & 0x7fff);
-			Packet->PacketSize = uint16(Size);
-
-			SendInner(Data, Size);
+			Buffer->Next = Head;
+			Head = Buffer;
+			Tail = (Tail != nullptr) ? Tail : Head;
 		}
 	};
 
@@ -517,7 +520,11 @@ static void Writer_ConsumeEvents()
 
 		if (uint32 SizeToReap = uint32(Committed - Buffer->Reaped))
 		{
-			Send(Buffer, Buffer->Reaped, SizeToReap);
+#if TRACE_PRIVATE_PERF
+			BytesReaped += SizeToReap;
+			BytesSend += /*...*/
+#endif
+			Writer_SendData(Buffer->ThreadId, Buffer->Reaped, SizeToReap);
 			Buffer->Reaped = Committed;
 		}
 
@@ -553,7 +560,11 @@ static void Writer_ConsumeEvents()
 
 		if (uint32 SizeToReap = uint32(Buffer->Committed - Buffer->Reaped))
 		{
-			Send(Buffer, Buffer->Reaped, SizeToReap);
+#if TRACE_PRIVATE_PERF
+			BytesReaped += SizeToReap;
+			BytesSend += /*...*/
+#endif
+			Writer_SendData(Buffer->ThreadId, Buffer->Reaped, SizeToReap);
 		}
 	}
 
@@ -571,6 +582,7 @@ static void Writer_ConsumeEvents()
 #if TRACE_PRIVATE_PERF
 	UE_TRACE_LOG($Trace, WorkerThread, TraceLogChannel)
 		<< WorkerThread.Cycles(uint32(TimeGetTimestamp() - StartTsc))
+		<< WorkerThread.BytesReaped(BytesReaped);
 		<< WorkerThread.BytesSent(BytesSent);
 
 	UE_TRACE_LOG($Trace, Memory, TraceLogChannel)
