Use a zero-valued thread ID as an indication that there's not valid trace event buffer yet.

#rnx

diff --git a/Engine/Source/Runtime/TraceLog/Private/Trace/Writer.cpp b/Engine/Source/Runtime/TraceLog/Private/Trace/Writer.cpp
--- a/Engine/Source/Runtime/TraceLog/Private/Trace/Writer.cpp
+++ b/Engine/Source/Runtime/TraceLog/Private/Trace/Writer.cpp
@@ -81,13 +81,14 @@ FWriteTlsContext::FWriteTlsContext()
 	if (!Once)
 	{
 		Target->Cursor = DefaultBuffer;
+		Target->ThreadId = 0;
 		Once = true;
 	}
 
 	Buffer = Target;
 
 	// Assign a new id to this thread.
-	ThreadId = AtomicIncrementRelaxed(&ThreadIdCounter);
+	ThreadId = AtomicIncrementRelaxed(&ThreadIdCounter) + 1;
 }
 
 ////////////////////////////////////////////////////////////////////////////////
@@ -103,7 +104,7 @@ FWriteTlsContext::~FWriteTlsContext()
 ////////////////////////////////////////////////////////////////////////////////
 bool FWriteTlsContext::HasValidBuffer() const
 {
-	return (UPTRINT(Buffer) != UPTRINT(DefaultBuffer));
+	return (Buffer->ThreadId != 0);
 }
 
 ////////////////////////////////////////////////////////////////////////////////
