Do not keep a thread ID in the TlsContext object so its machine-word size.

#rnx

diff --git a/Engine/Source/Runtime/TraceLog/Private/Trace/Writer.cpp b/Engine/Source/Runtime/TraceLog/Private/Trace/Writer.cpp
--- a/Engine/Source/Runtime/TraceLog/Private/Trace/Writer.cpp
+++ b/Engine/Source/Runtime/TraceLog/Private/Trace/Writer.cpp
@@ -86,9 +86,6 @@ FWriteTlsContext::FWriteTlsContext()
 	}
 
 	Buffer = Target;
-
-	// Assign a new id to this thread.
-	ThreadId = AtomicIncrementRelaxed(&ThreadIdCounter) + 1;
 }
 
 ////////////////////////////////////////////////////////////////////////////////
@@ -110,13 +107,18 @@ bool FWriteTlsContext::HasValidBuffer() const
 ////////////////////////////////////////////////////////////////////////////////
 inline void FWriteTlsContext::SetBuffer(FWriteBuffer* InBuffer)
 {
-	Buffer = InBuffer;
-}
+	uint32 ThreadId;
+	if (!Buffer->ThreadId)
+	{
+		ThreadId = AtomicIncrementRelaxed(&ThreadIdCounter) + 1;
+	}
+	else
+	{
+		ThreadId = Buffer->ThreadId;
+	}
 
-////////////////////////////////////////////////////////////////////////////////
-inline uint32 FWriteTlsContext::GetThreadId() const
-{
-	return ThreadId;
+	Buffer = InBuffer;
+	Buffer->ThreadId = ThreadId;
 }
 
 
@@ -249,7 +251,6 @@ TRACELOG_API FWriteBuffer* Writer_NextBuffer(uint16 Size)
 	}
 
 	FWriteBuffer* NextBuffer = Writer_NextBufferInternal(GPoolPageGrowth);
-	NextBuffer->ThreadId = TlsContext.GetThreadId();
 
 	NextBuffer->Cursor += Size;
 	return NextBuffer;
diff --git a/Engine/Source/Runtime/TraceLog/Public/Trace/Detail/Writer.inl b/Engine/Source/Runtime/TraceLog/Public/Trace/Detail/Writer.inl
--- a/Engine/Source/Runtime/TraceLog/Public/Trace/Detail/Writer.inl
+++ b/Engine/Source/Runtime/TraceLog/Public/Trace/Detail/Writer.inl
@@ -34,12 +34,10 @@ struct FWriteTlsContext
 							~FWriteTlsContext();
 	bool					HasValidBuffer() const;
 	void					SetBuffer(FWriteBuffer*);
-	uint32					GetThreadId() const;
 	FWriteBuffer*			GetBuffer() const { return Buffer; }
 
 private:
 	FWriteBuffer*			Buffer;
-	uint32					ThreadId;
 	static uint8			DefaultBuffer[sizeof(FWriteBuffer)];
 	static uint32 volatile	ThreadIdCounter;
 };
