Organise trace event buffers into lists per thread. This allows us to only send the oldest of a thread's event data first, and never process future buffers until prior ones are complete.

#rnx

diff --git a/Engine/Source/Runtime/TraceLog/Private/Trace/Writer.cpp b/Engine/Source/Runtime/TraceLog/Private/Trace/Writer.cpp
--- a/Engine/Source/Runtime/TraceLog/Private/Trace/Writer.cpp
+++ b/Engine/Source/Runtime/TraceLog/Private/Trace/Writer.cpp
@@ -93,7 +93,7 @@ FWriteTlsContext::~FWriteTlsContext()
 {
 	if (GInitialized && HasValidBuffer())
 	{
-		UPTRINT EtxOffset = ~UPTRINT((uint8*)Buffer - Buffer->Cursor);
+		UPTRINT EtxOffset = UPTRINT((uint8*)Buffer - Buffer->Cursor);
 		AtomicStoreRelaxed(&(Buffer->EtxOffset), EtxOffset);
 	}
 }
@@ -105,23 +105,26 @@ bool FWriteTlsContext::HasValidBuffer() const
 }
 
 ////////////////////////////////////////////////////////////////////////////////
-inline void FWriteTlsContext::SetBuffer(FWriteBuffer* InBuffer)
+inline FWriteBuffer* FWriteTlsContext::SetBuffer(FWriteBuffer* InBuffer)
 {
+	FWriteBuffer* PrevBuffer = Buffer;
+
 	uint32 ThreadId;
-	if (!Buffer->ThreadId)
+	if (!PrevBuffer->ThreadId)
 	{
 		ThreadId = AtomicIncrementRelaxed(&ThreadIdCounter) + 1;
+		PrevBuffer = nullptr;
 	}
 	else
 	{
-		ThreadId = Buffer->ThreadId;
+		ThreadId = PrevBuffer->ThreadId;
 	}
 
 	Buffer = InBuffer;
 	Buffer->ThreadId = ThreadId;
-}
-
 
+	return PrevBuffer;
+}
 
 ////////////////////////////////////////////////////////////////////////////////
 #define T_ALIGN alignas(PLATFORM_CACHE_LINE_SIZE)
@@ -132,7 +135,7 @@ static const uint32						GPoolInitPageSize	= GPoolBlockSize << 5;
 static uint8*							GPoolBase;			// = nullptr;
 T_ALIGN static uint8* volatile			GPoolPageCursor;	// = nullptr;
 T_ALIGN static FWriteBuffer* volatile	GPoolFreeList;		// = nullptr;
-T_ALIGN static FWriteBuffer* volatile	GNextBufferList;	// = nullptr;
+T_ALIGN static FWriteBuffer* volatile	GNewThreadList;		// = nullptr;
 #undef T_ALIGN
 
 ////////////////////////////////////////////////////////////////////////////////
@@ -156,7 +159,7 @@ static FWriteBuffer* Writer_NextBufferInternal(uint32 PageGrowth)
 		FWriteBuffer* Owned = AtomicLoadRelaxed(&GPoolFreeList);
 		if (Owned != nullptr)
 		{
-			if (!AtomicCompareExchangeRelaxed(&GPoolFreeList, Owned->Next, Owned))
+			if (!AtomicCompareExchangeRelaxed(&GPoolFreeList, Owned->NextBuffer, Owned))
 			{
 				Private::PlatformYield();
 				continue;
@@ -195,15 +198,15 @@ static FWriteBuffer* Writer_NextBufferInternal(uint32 PageGrowth)
 		for (int i = 2, n = PageGrowth / GPoolBlockSize; i < n; ++i)
 		{
 			auto* Buffer = (FWriteBuffer*)Block;
-			Buffer->Next = (FWriteBuffer*)(Block + GPoolBlockSize);
+			Buffer->NextBuffer = (FWriteBuffer*)(Block + GPoolBlockSize);
 			Block += GPoolBlockSize;
 		}
 
 		// And insert the block list into the freelist. 'Block' is now the last block
 		for (auto* ListNode = (FWriteBuffer*)Block;; Private::PlatformYield())
 		{
-			ListNode->Next = AtomicLoadRelaxed(&GPoolFreeList);
-			if (AtomicCompareExchangeRelease(&GPoolFreeList, (FWriteBuffer*)FirstBlock, ListNode->Next))
+			ListNode->NextBuffer = AtomicLoadRelaxed(&GPoolFreeList);
+			if (AtomicCompareExchangeRelease(&GPoolFreeList, (FWriteBuffer*)FirstBlock, ListNode->NextBuffer))
 			{
 				break;
 			}
@@ -216,19 +219,32 @@ static FWriteBuffer* Writer_NextBufferInternal(uint32 PageGrowth)
 	NextBuffer->Cursor += sizeof(uint32); // this is so we can preceed event data with a small header when sending.
 	NextBuffer->Committed = NextBuffer->Cursor;
 	NextBuffer->Reaped = NextBuffer->Cursor;
-	NextBuffer->EtxOffset = 0;
+	NextBuffer->EtxOffset = UPTRINT(0) - sizeof(FWriteBuffer);
+	NextBuffer->NextBuffer = nullptr;
 
-	// Add this next buffer to the active list.
-	for (;; Private::PlatformYield())
+
+	FWriteBuffer* PrevBuffer = TlsContext.SetBuffer(NextBuffer);
+	if (PrevBuffer == nullptr)
 	{
-		NextBuffer->Next = AtomicLoadRelaxed(&GNextBufferList);
-		if (AtomicCompareExchangeRelease(&GNextBufferList, NextBuffer, NextBuffer->Next))
+		// Add this next buffer to the active list.
+		for (;; Private::PlatformYield())
 		{
-			break;
+			NextBuffer->NextThread = AtomicLoadRelaxed(&GNewThreadList);
+			if (AtomicCompareExchangeRelease(&GNewThreadList, NextBuffer, NextBuffer->NextThread))
+			{
+				break;
+			}
 		}
 	}
+	else
+	{
+		PrevBuffer->NextBuffer = NextBuffer;
+
+		// Retire current buffer.
+		UPTRINT EtxOffset = UPTRINT((uint8*)(PrevBuffer) - PrevBuffer->Cursor);
+		AtomicStoreRelease(&(PrevBuffer->EtxOffset), EtxOffset);
+	}
 
-	TlsContext.SetBuffer(NextBuffer);
 	return NextBuffer;
 }
 
@@ -241,13 +257,11 @@ TRACELOG_API FWriteBuffer* Writer_NextBuffer(uint16 Size)
 		return nullptr;
 	}
 
-	FWriteBuffer* Current = TlsContext.GetBuffer();
-
 	// Retire current buffer unless its the initial boot one.
 	if (TlsContext.HasValidBuffer())
 	{
-		UPTRINT EtxOffset = ~UPTRINT((uint8*)(Current) - Current->Cursor + Size);
-		AtomicStoreRelease(&(Current->EtxOffset), EtxOffset);
+		FWriteBuffer* Current = TlsContext.GetBuffer();
+		Current->Cursor -= Size;
 	}
 
 	FWriteBuffer* NextBuffer = Writer_NextBufferInternal(GPoolPageGrowth);
@@ -371,11 +385,11 @@ enum class EDataState : uint8
 	Partial,			// Passive, but buffers are full so some events are lost
 	Sending,			// Events are being sent to an IO handle
 };
-static FHoldBuffer		GHoldBuffer;		// will init to zero.
-static UPTRINT			GDataHandle;		// = 0
-static EDataState		GDataState;			// = EDataState::Passive
-UPTRINT					GPendingDataHandle;	// = 0
-static FWriteBuffer*	GActiveBufferList;	// = nullptr;
+static FHoldBuffer				GHoldBuffer;		// will init to zero.
+static UPTRINT					GDataHandle;		// = 0
+static EDataState				GDataState;			// = EDataState::Passive
+UPTRINT							GPendingDataHandle;	// = 0
+static FWriteBuffer* __restrict GActiveThreadList;	// = nullptr;
 
 ////////////////////////////////////////////////////////////////////////////////
 static uint32 Writer_SendData(uint32 ThreadId, uint8* __restrict Data, uint32 Size)
@@ -453,22 +467,6 @@ static uint32 Writer_SendData(uint32 ThreadId, uint8* __restrict Data, uint32 Si
 ////////////////////////////////////////////////////////////////////////////////
 static void Writer_ConsumeEvents()
 {
-#if TRACE_PRIVATE_PERF
-	uint64 StartTsc = TimeGetTimestamp();
-	uint32 BytesSent = 0;
-#endif
-
-	// Claim ownership of the latest chain of sent events.
-	FWriteBuffer* __restrict NextBufferList;
-	for (;; Private::PlatformYield())
-	{
-		NextBufferList = AtomicLoadRelaxed(&GNextBufferList);
-		if (AtomicCompareExchangeAcquire(&GNextBufferList, (FWriteBuffer*)nullptr, NextBufferList))
-		{
-			break;
-		}
-	}
-
 	struct FRetireList
 	{
 		FWriteBuffer* __restrict Head = nullptr;
@@ -476,111 +474,80 @@ static void Writer_ConsumeEvents()
 
 		void Insert(FWriteBuffer* __restrict Buffer)
 		{
-			Buffer->Next = Head;
+			Buffer->NextBuffer = Head;
 			Head = Buffer;
 			Tail = (Tail != nullptr) ? Tail : Head;
 		}
 	};
 
-	FRetireList RetireList;
+#if TRACE_PRIVATE_PERF
+	uint64 StartTsc = TimeGetTimestamp();
+	uint32 BytesReaped = 0;
+	uint32 BytesSent = 0;
+#endif
 
-	// Next buffer list is newest first. Retire full ones and build a new list
-	// of buffers that are active (which gets reverse so oldest is first)
-	FWriteBuffer* __restrict NextActiveList = nullptr;
-	FWriteBuffer* __restrict NextRetireList = nullptr;
-	for (FWriteBuffer *__restrict Buffer = NextBufferList, *__restrict NextBuffer;
-		Buffer != nullptr;
-		Buffer = NextBuffer)
+	// Claim ownership of any new thread buffer lists
+	FWriteBuffer* __restrict NewThreadList;
+	for (;; Private::PlatformYield())
 	{
-		NextBuffer = Buffer->Next;
-
-		uint8* Committed = AtomicLoadAcquire(&Buffer->Committed);
-		int32 EtxOffset = int32(~AtomicLoadRelaxed(&Buffer->EtxOffset));
-		if ((uint8*)Buffer - EtxOffset > Committed)
-		{
-			Buffer->Next = NextActiveList;
-			NextActiveList = Buffer;
-		}
-		else
+		NewThreadList = AtomicLoadRelaxed(&GNewThreadList);
+		if (AtomicCompareExchangeAcquire(&GNewThreadList, (FWriteBuffer*)nullptr, NewThreadList))
 		{
-			Buffer->Next = NextRetireList;
-			NextRetireList = Buffer;
+			break;
 		}
 	}
 
-	// Send as much of the active list as we can. Buffers that are full are removed
-	// from the list. Note that the list's oldest-first order is maintained
-	FWriteBuffer* __restrict ActiveListHead = nullptr;
-	FWriteBuffer* __restrict ActiveListTail = nullptr;
-	for (FWriteBuffer *__restrict Buffer = GActiveBufferList, *__restrict NextBuffer;
-		Buffer != nullptr;
-		Buffer = NextBuffer)
-	{
-		NextBuffer = Buffer->Next;
+	FRetireList RetireList;
 
-		uint8* __restrict Committed = AtomicLoadAcquire(&Buffer->Committed);
+	FWriteBuffer* __restrict ActiveThreadList = GActiveThreadList;
+	GActiveThreadList = nullptr;
 
-		if (uint32 SizeToReap = uint32(Committed - Buffer->Reaped))
+	// Now we've two lists of known and new threads. Each of these lists in turn is
+	// a list of that thread's buffers (where it is writing trace events to).
+	for (FWriteBuffer* __restrict Buffer : { ActiveThreadList, NewThreadList })
+	{
+		// For each thread...
+		for (FWriteBuffer* __restrict NextThread; Buffer != nullptr; Buffer = NextThread)
 		{
-#if TRACE_PRIVATE_PERF
-			BytesReaped += SizeToReap;
-			BytesSend += /*...*/
-#endif
-			Writer_SendData(Buffer->ThreadId, Buffer->Reaped, SizeToReap);
-			Buffer->Reaped = Committed;
-		}
+			NextThread = Buffer->NextThread;
+			uint32 ThreadId = Buffer->ThreadId;
 
-		int32 EtxOffset = int32(~AtomicLoadRelaxed(&Buffer->EtxOffset));
-		if ((uint8*)Buffer - EtxOffset == Committed)
-		{
-			RetireList.Insert(Buffer);
-		}
-		else
-		{
-			if (ActiveListTail != nullptr)
-			{
-				ActiveListTail->Next = Buffer;
-			}
-			else
+			// For each of the thread's buffers...
+			for (FWriteBuffer* __restrict NextBuffer; Buffer != nullptr; Buffer = NextBuffer)
 			{
-				ActiveListHead = Buffer;
-			}
+				uint8* Committed = AtomicLoadRelaxed(&Buffer->Committed);
 
-			ActiveListTail = Buffer;
-			Buffer->Next = nullptr;
-		}
-	}
+				// Send as much as we can.
+				if (uint32 SizeToReap = uint32(Committed - Buffer->Reaped))
+				{
+#if TRACE_PRIVATE_PERF
+					BytesReaped += SizeToReap;
+					BytesSent += /*...*/
+#endif
+					Writer_SendData(ThreadId, Buffer->Reaped, SizeToReap);
+					Buffer->Reaped = Committed;
+				}
 
-	// Retire buffers from the next list
-	for (FWriteBuffer *__restrict Buffer = NextRetireList, *__restrict NextBuffer;
-		Buffer != nullptr;
-		Buffer = NextBuffer)
-	{
-		NextBuffer = Buffer->Next;
+				// Is this buffer still in use?
+				int32 EtxOffset = int32(AtomicLoadAcquire(&Buffer->EtxOffset));
+				if ((uint8*)Buffer - EtxOffset > Committed)
+				{
+					break;
+				}
 
-		RetireList.Insert(Buffer);
+				// Retire the buffer
+				NextBuffer = Buffer->NextBuffer;
+				RetireList.Insert(Buffer);
+			}
 
-		if (uint32 SizeToReap = uint32(Buffer->Committed - Buffer->Reaped))
-		{
-#if TRACE_PRIVATE_PERF
-			BytesReaped += SizeToReap;
-			BytesSend += /*...*/
-#endif
-			Writer_SendData(Buffer->ThreadId, Buffer->Reaped, SizeToReap);
+			if (Buffer != nullptr)
+			{
+				Buffer->NextThread = GActiveThreadList;
+				GActiveThreadList = Buffer;
+			}
 		}
 	}
 
-	// Append the new active buffers that have been discovered to the active list
-	if (ActiveListTail != nullptr)
-	{
-		GActiveBufferList = ActiveListHead;
-		ActiveListTail->Next = NextActiveList;
-	}
-	else
-	{
-		GActiveBufferList = NextActiveList;
-	}
-
 #if TRACE_PRIVATE_PERF
 	UE_TRACE_LOG($Trace, WorkerThread, TraceLogChannel)
 		<< WorkerThread.Cycles(uint32(TimeGetTimestamp() - StartTsc))
@@ -596,8 +563,8 @@ static void Writer_ConsumeEvents()
 	{
 		for (FWriteBuffer* ListNode = RetireList.Tail;; Private::PlatformYield())
 		{
-			ListNode->Next = AtomicLoadRelaxed(&GPoolFreeList);
-			if (AtomicCompareExchangeRelease(&GPoolFreeList, RetireList.Head, ListNode->Next))
+			ListNode->NextBuffer = AtomicLoadRelaxed(&GPoolFreeList);
+			if (AtomicCompareExchangeRelease(&GPoolFreeList, RetireList.Head, ListNode->NextBuffer))
 			{
 				break;
 			}
diff --git a/Engine/Source/Runtime/TraceLog/Public/Trace/Detail/Writer.inl b/Engine/Source/Runtime/TraceLog/Public/Trace/Detail/Writer.inl
--- a/Engine/Source/Runtime/TraceLog/Public/Trace/Detail/Writer.inl
+++ b/Engine/Source/Runtime/TraceLog/Public/Trace/Detail/Writer.inl
@@ -18,7 +18,8 @@ struct FWriteBuffer
 {
 	uint32						Overflow;
 	uint32						ThreadId;
-	FWriteBuffer* __restrict	Next;
+	FWriteBuffer* __restrict	NextThread;
+	FWriteBuffer* __restrict	NextBuffer;
 	uint8* __restrict			Cursor;
 	uint8* __restrict volatile	Committed;
 	uint8* __restrict			Reaped;
@@ -33,7 +34,7 @@ struct FWriteTlsContext
 							FWriteTlsContext();
 							~FWriteTlsContext();
 	bool					HasValidBuffer() const;
-	void					SetBuffer(FWriteBuffer*);
+	FWriteBuffer*			SetBuffer(FWriteBuffer*);
 	FWriteBuffer*			GetBuffer() const { return Buffer; }
 
 private:
