Moved the runtime-side command/control handling out into its own translation unit.

#rnx

diff --git a/Engine/Source/Runtime/TraceLog/Private/Trace/Control.cpp b/Engine/Source/Runtime/TraceLog/Private/Trace/Control.cpp
new file mode 100644
--- /dev/null
+++ b/Engine/Source/Runtime/TraceLog/Private/Trace/Control.cpp
@@ -0,0 +1,317 @@
+// Copyright Epic Games, Inc. All Rights Reserved.
+
+#include "Trace/Config.h"
+
+#if UE_TRACE_ENABLED
+
+#include "Trace/Platform.h"
+
+#include "Misc/CString.h"
+#include "Templates/UnrealTemplate.h"
+
+namespace Trace {
+namespace Private {
+
+////////////////////////////////////////////////////////////////////////////////
+bool	Writer_SendTo(const ANSICHAR*, uint32);
+bool	Writer_WriteTo(const ANSICHAR*);
+
+
+
+////////////////////////////////////////////////////////////////////////////////
+enum class EControlState : uint8
+{
+	Closed = 0,
+	Listening,
+	Accepted,
+	Failed,
+};
+
+////////////////////////////////////////////////////////////////////////////////
+struct FControlCommands
+{
+	enum { Max = 8 };
+	struct
+	{
+		uint32	Hash;
+		void*	Param;
+		void	(*Thunk)(void*, uint32, ANSICHAR const* const*);
+	}			Commands[Max];
+	uint8		Count;
+};
+
+
+
+////////////////////////////////////////////////////////////////////////////////
+static FControlCommands	GControlCommands;
+static UPTRINT			GControlListen		= 0;
+static UPTRINT			GControlSocket		= 0;
+static EControlState	GControlState;		// = EControlState::Closed;
+
+////////////////////////////////////////////////////////////////////////////////
+static uint32 Writer_ControlHash(const ANSICHAR* Word)
+{
+	uint32 Hash = 5381;
+	for (; *Word; (Hash = (Hash * 33) ^ *Word), ++Word);
+	return Hash;
+}
+
+////////////////////////////////////////////////////////////////////////////////
+static bool Writer_ControlAddCommand(
+	const ANSICHAR* Name,
+	void* Param,
+	void (*Thunk)(void*, uint32, ANSICHAR const* const*))
+{
+	if (GControlCommands.Count >= FControlCommands::Max)
+	{
+		return false;
+	}
+
+	uint32 Index = GControlCommands.Count++;
+	GControlCommands.Commands[Index] = { Writer_ControlHash(Name), Param, Thunk };
+	return true;
+}
+
+////////////////////////////////////////////////////////////////////////////////
+static bool Writer_ControlDispatch(uint32 ArgC, ANSICHAR const* const* ArgV)
+{
+	if (ArgC == 0)
+	{
+		return false;
+	}
+
+	uint32 Hash = Writer_ControlHash(ArgV[0]);
+	--ArgC;
+	++ArgV;
+
+	for (int i = 0, n = GControlCommands.Count; i < n; ++i)
+	{
+		const auto& Command = GControlCommands.Commands[i];
+		if (Command.Hash == Hash)
+		{
+			Command.Thunk(Command.Param, ArgC, ArgV);
+			return true;
+		}
+	}
+
+	return false;
+}
+
+////////////////////////////////////////////////////////////////////////////////
+static bool Writer_ControlListen()
+{
+	GControlListen = TcpSocketListen(1985);
+	if (!GControlListen)
+	{
+		GControlState = EControlState::Failed;
+		return false;
+	}
+
+	GControlState = EControlState::Listening;
+	return true;
+}
+
+////////////////////////////////////////////////////////////////////////////////
+static bool Writer_ControlAccept()
+{
+	UPTRINT Socket;
+	int Return = TcpSocketAccept(GControlListen, Socket);
+	if (Return <= 0)
+	{
+		if (Return == -1)
+		{
+			IoClose(GControlListen);
+			GControlListen = 0;
+			GControlState = EControlState::Failed;
+		}
+		return false;
+	}
+
+	GControlState = EControlState::Accepted;
+	GControlSocket = Socket;
+	return true;
+}
+
+////////////////////////////////////////////////////////////////////////////////
+static void Writer_ControlRecv()
+{
+	// We'll assume that commands are smaller than the canonical MTU so this
+	// doesn't need to be implemented in a reentrant manner (maybe).
+
+	ANSICHAR Buffer[512];
+	ANSICHAR* __restrict Head = Buffer;
+	while (TcpSocketHasData(GControlSocket))
+	{
+		int32 ReadSize = int32(UPTRINT(Buffer + sizeof(Buffer) - Head));
+		int32 Recvd = IoRead(GControlSocket, Head, ReadSize);
+		if (Recvd <= 0)
+		{
+			IoClose(GControlSocket);
+			GControlSocket = 0;
+			GControlState = EControlState::Listening;
+			break;
+		}
+
+		Head += Recvd;
+
+		enum EParseState
+		{
+			CrLfSkip,
+			WhitespaceSkip,
+			Word,
+		} ParseState = EParseState::CrLfSkip;
+
+		uint32 ArgC = 0;
+		const ANSICHAR* ArgV[16];
+
+		const ANSICHAR* __restrict Spent = Buffer;
+		for (ANSICHAR* __restrict Cursor = Buffer; Cursor < Head; ++Cursor)
+		{
+			switch (ParseState)
+			{
+			case EParseState::CrLfSkip:
+				if (*Cursor == '\n' || *Cursor == '\r')
+				{
+					continue;
+				}
+				ParseState = EParseState::WhitespaceSkip;
+				/* [[fallthrough]] */
+
+			case EParseState::WhitespaceSkip:
+				if (*Cursor == ' ' || *Cursor == '\0')
+				{
+					continue;
+				}
+
+				if (ArgC < UE_ARRAY_COUNT(ArgV))
+				{
+					ArgV[ArgC] = Cursor;
+					++ArgC;
+				}
+
+				ParseState = EParseState::Word;
+				/* [[fallthrough]] */
+
+			case EParseState::Word:
+				if (*Cursor == ' ' || *Cursor == '\0')
+				{
+					*Cursor = '\0';
+					ParseState = EParseState::WhitespaceSkip;
+					continue;
+				}
+
+				if (*Cursor == '\r' || *Cursor == '\n')
+				{
+					*Cursor = '\0';
+
+					Writer_ControlDispatch(ArgC, ArgV);
+
+					ArgC = 0;
+					Spent = Cursor + 1;
+					ParseState = EParseState::CrLfSkip;
+					continue;
+				}
+
+				break;
+			}
+		}
+
+		int32 UnspentSize = int32(UPTRINT(Head - Spent));
+		if (UnspentSize)
+		{
+			memmove(Buffer, Spent, UnspentSize);
+		}
+		Head = Buffer + UnspentSize;
+	}
+}
+
+////////////////////////////////////////////////////////////////////////////////
+void Writer_UpdateControl()
+{
+	switch (GControlState)
+	{
+	case EControlState::Closed:
+		if (!Writer_ControlListen())
+		{
+			break;
+		}
+		/* [[fallthrough]] */
+
+	case EControlState::Listening:
+		if (!Writer_ControlAccept())
+		{
+			break;
+		}
+		/* [[fallthrough]] */
+
+	case EControlState::Accepted:
+		Writer_ControlRecv();
+		break;
+	}
+}
+
+////////////////////////////////////////////////////////////////////////////////
+void Writer_InitializeControl()
+{
+#if PLATFORM_SWITCH
+	GControlState = EControlState::Failed;
+	return;
+#endif
+
+	Writer_ControlAddCommand("SendTo", nullptr,
+		[] (void*, uint32 ArgC, ANSICHAR const* const* ArgV)
+		{
+			if (ArgC > 0)
+			{
+				Writer_SendTo(ArgV[0], 1980);
+			}
+		}
+	);
+
+	Writer_ControlAddCommand("WriteTo", nullptr,
+		[] (void*, uint32 ArgC, ANSICHAR const* const* ArgV)
+		{
+			if (ArgC > 0)
+			{
+				Writer_WriteTo(ArgV[0]);
+			}
+		}
+	);
+
+	Writer_ControlAddCommand("ToggleChannels", nullptr, 
+		[] (void*, uint32 ArgC, ANSICHAR const* const* ArgV) 
+		{
+			if (ArgC < 2)
+			{
+				return;
+			}
+
+			const size_t BufferSize = 512;
+			ANSICHAR Channels[BufferSize] = {};
+			ANSICHAR* Ctx;
+			const bool bState = (ArgV[1][0] != '0');
+			FPlatformString::Strcpy(Channels, BufferSize, ArgV[0]);
+			ANSICHAR* Channel = FPlatformString::Strtok(Channels, ",", &Ctx);
+			while (Channel)
+			{
+				FChannel::Toggle(Channel, bState);
+				Channel = FPlatformString::Strtok(nullptr, ",", &Ctx);
+			}
+		}
+	);
+}
+
+////////////////////////////////////////////////////////////////////////////////
+void Writer_ShutdownControl()
+{
+	if (GControlListen)
+	{
+		IoClose(GControlListen);
+		GControlListen = 0;
+	}
+}
+
+} // namespace Private
+} // namespace Trace
+
+#endif // UE_TRACE_ENABLED
diff --git a/Engine/Source/Runtime/TraceLog/Private/Trace/Writer.cpp b/Engine/Source/Runtime/TraceLog/Private/Trace/Writer.cpp
--- a/Engine/Source/Runtime/TraceLog/Private/Trace/Writer.cpp
+++ b/Engine/Source/Runtime/TraceLog/Private/Trace/Writer.cpp
@@ -10,16 +10,14 @@
 #include "Trace/Platform.h"
 #include "Trace/Trace.h"
 
-#include "Misc/CString.h"
-#include "Templates/UnrealTemplate.h"
-
-
-
 namespace Trace {
 namespace Private {
 
 ////////////////////////////////////////////////////////////////////////////////
-int32 Encode(const void*, int32, void*, int32);
+int32	Encode(const void*, int32, void*, int32);
+void	Writer_UpdateControl();
+void	Writer_InitializeControl();
+void	Writer_ShutdownControl();
 
 
 
@@ -622,302 +620,6 @@ static void Writer_UpdateData()
 
 
 
-////////////////////////////////////////////////////////////////////////////////
-enum class EControlState : uint8
-{
-	Closed = 0,
-	Listening,
-	Accepted,
-	Failed,
-};
-
-struct FControlCommands
-{
-	enum { Max = 16 };
-	struct
-	{
-		uint32	Hash;
-		void*	Param;
-		void	(*Thunk)(void*, uint32, ANSICHAR const* const*);
-	}			Commands[Max];
-	uint8		Count;
-};
-
-////////////////////////////////////////////////////////////////////////////////
-bool	Writer_SendTo(const ANSICHAR*, uint32);
-bool	Writer_WriteTo(const ANSICHAR*);
-
-////////////////////////////////////////////////////////////////////////////////
-static FControlCommands	GControlCommands;
-static UPTRINT			GControlListen		= 0;
-static UPTRINT			GControlSocket		= 0;
-static EControlState	GControlState;		// = EControlState::Closed;
-
-////////////////////////////////////////////////////////////////////////////////
-static uint32 Writer_ControlHash(const ANSICHAR* Word)
-{
-	uint32 Hash = 5381;
-	for (; *Word; (Hash = (Hash * 33) ^ *Word), ++Word);
-	return Hash;
-}
-
-////////////////////////////////////////////////////////////////////////////////
-static bool Writer_ControlAddCommand(
-	const ANSICHAR* Name,
-	void* Param,
-	void (*Thunk)(void*, uint32, ANSICHAR const* const*))
-{
-	if (GControlCommands.Count >= FControlCommands::Max)
-	{
-		return false;
-	}
-
-	uint32 Index = GControlCommands.Count++;
-	GControlCommands.Commands[Index] = { Writer_ControlHash(Name), Param, Thunk };
-	return true;
-}
-
-////////////////////////////////////////////////////////////////////////////////
-static bool Writer_ControlDispatch(uint32 ArgC, ANSICHAR const* const* ArgV)
-{
-	if (ArgC == 0)
-	{
-		return false;
-	}
-
-	uint32 Hash = Writer_ControlHash(ArgV[0]);
-	--ArgC;
-	++ArgV;
-
-	for (int i = 0, n = GControlCommands.Count; i < n; ++i)
-	{
-		const auto& Command = GControlCommands.Commands[i];
-		if (Command.Hash == Hash)
-		{
-			Command.Thunk(Command.Param, ArgC, ArgV);
-			return true;
-		}
-	}
-
-	return false;
-}
-
-////////////////////////////////////////////////////////////////////////////////
-static bool Writer_ControlListen()
-{
-	GControlListen = TcpSocketListen(1985);
-	if (!GControlListen)
-	{
-		GControlState = EControlState::Failed;
-		return false;
-	}
-
-	GControlState = EControlState::Listening;
-	return true;
-}
-
-////////////////////////////////////////////////////////////////////////////////
-static bool Writer_ControlAccept()
-{
-	UPTRINT Socket;
-	int Return = TcpSocketAccept(GControlListen, Socket);
-	if (Return <= 0)
-	{
-		if (Return == -1)
-		{
-			IoClose(GControlListen);
-			GControlListen = 0;
-			GControlState = EControlState::Failed;
-		}
-		return false;
-	}
-
-	GControlState = EControlState::Accepted;
-	GControlSocket = Socket;
-	return true;
-}
-
-////////////////////////////////////////////////////////////////////////////////
-static void Writer_ControlRecv()
-{
-	// We'll assume that commands are smaller than the canonical MTU so this
-	// doesn't need to be implemented in a reentrant manner (maybe).
-
-	ANSICHAR Buffer[512];
-	ANSICHAR* __restrict Head = Buffer;
-	while (TcpSocketHasData(GControlSocket))
-	{
-		int32 ReadSize = int32(UPTRINT(Buffer + sizeof(Buffer) - Head));
-		int32 Recvd = IoRead(GControlSocket, Head, ReadSize);
-		if (Recvd <= 0)
-		{
-			IoClose(GControlSocket);
-			GControlSocket = 0;
-			GControlState = EControlState::Listening;
-			break;
-		}
-
-		Head += Recvd;
-
-		enum EParseState
-		{
-			CrLfSkip,
-			WhitespaceSkip,
-			Word,
-		} ParseState = EParseState::CrLfSkip;
-
-		uint32 ArgC = 0;
-		const ANSICHAR* ArgV[16];
-
-		const ANSICHAR* __restrict Spent = Buffer;
-		for (ANSICHAR* __restrict Cursor = Buffer; Cursor < Head; ++Cursor)
-		{
-			switch (ParseState)
-			{
-			case EParseState::CrLfSkip:
-				if (*Cursor == '\n' || *Cursor == '\r')
-				{
-					continue;
-				}
-				ParseState = EParseState::WhitespaceSkip;
-				/* [[fallthrough]] */
-
-			case EParseState::WhitespaceSkip:
-				if (*Cursor == ' ' || *Cursor == '\0')
-				{
-					continue;
-				}
-
-				if (ArgC < UE_ARRAY_COUNT(ArgV))
-				{
-					ArgV[ArgC] = Cursor;
-					++ArgC;
-				}
-
-				ParseState = EParseState::Word;
-				/* [[fallthrough]] */
-
-			case EParseState::Word:
-				if (*Cursor == ' ' || *Cursor == '\0')
-				{
-					*Cursor = '\0';
-					ParseState = EParseState::WhitespaceSkip;
-					continue;
-				}
-
-				if (*Cursor == '\r' || *Cursor == '\n')
-				{
-					*Cursor = '\0';
-
-					Writer_ControlDispatch(ArgC, ArgV);
-
-					ArgC = 0;
-					Spent = Cursor + 1;
-					ParseState = EParseState::CrLfSkip;
-					continue;
-				}
-
-				break;
-			}
-		}
-
-		int32 UnspentSize = int32(UPTRINT(Head - Spent));
-		if (UnspentSize)
-		{
-			memmove(Buffer, Spent, UnspentSize);
-		}
-		Head = Buffer + UnspentSize;
-	}
-}
-
-////////////////////////////////////////////////////////////////////////////////
-static void Writer_UpdateControl()
-{
-	switch (GControlState)
-	{
-	case EControlState::Closed:
-		if (!Writer_ControlListen())
-		{
-			break;
-		}
-		/* [[fallthrough]] */
-
-	case EControlState::Listening:
-		if (!Writer_ControlAccept())
-		{
-			break;
-		}
-		/* [[fallthrough]] */
-
-	case EControlState::Accepted:
-		Writer_ControlRecv();
-		break;
-	}
-}
-
-////////////////////////////////////////////////////////////////////////////////
-static void Writer_InitializeControl()
-{
-#if PLATFORM_SWITCH
-	GControlState = EControlState::Failed;
-	return;
-#endif
-
-	Writer_ControlAddCommand("SendTo", nullptr,
-		[] (void*, uint32 ArgC, ANSICHAR const* const* ArgV)
-		{
-			if (ArgC > 0)
-			{
-				Writer_SendTo(ArgV[0], 1980);
-			}
-		}
-	);
-
-	Writer_ControlAddCommand("WriteTo", nullptr,
-		[] (void*, uint32 ArgC, ANSICHAR const* const* ArgV)
-		{
-			if (ArgC > 0)
-			{
-				Writer_WriteTo(ArgV[0]);
-			}
-		}
-	);
-
-	Writer_ControlAddCommand("ToggleChannels", nullptr, 
-		[] (void*, uint32 ArgC, ANSICHAR const* const* ArgV) 
-		{
-			if (ArgC < 2)
-			{
-				return;
-			}
-
-			const size_t BufferSize = 512;
-			ANSICHAR Channels[BufferSize] = {};
-			ANSICHAR* Ctx;
-			const bool bState = (ArgV[1][0] != '0');
-			FPlatformString::Strcpy(Channels, BufferSize, ArgV[0]);
-			ANSICHAR* Channel = FPlatformString::Strtok(Channels, ",", &Ctx);
-			while (Channel)
-			{
-				FChannel::Toggle(Channel, bState);
-				Channel = FPlatformString::Strtok(nullptr, ",", &Ctx);
-			}
-		}
-	);
-}
-
-////////////////////////////////////////////////////////////////////////////////
-static void Writer_ShutdownControl()
-{
-	if (GControlListen)
-	{
-		IoClose(GControlListen);
-		GControlListen = 0;
-	}
-}
-
-
-
 ////////////////////////////////////////////////////////////////////////////////
 static UPTRINT			GWorkerThread		= 0;
 static volatile bool	GWorkerThreadQuit	= false;
@@ -938,6 +640,7 @@ static void Writer_WorkerThread()
 }
 
 
+
 ////////////////////////////////////////////////////////////////////////////////
 static void Writer_LogHeader()
 {
