Compile fixes

diff --git a/Engine/Source/Runtime/Engine/Public/ObjectTrace.h b/Engine/Source/Runtime/Engine/Public/ObjectTrace.h
--- a/Engine/Source/Runtime/Engine/Public/ObjectTrace.h
+++ b/Engine/Source/Runtime/Engine/Public/ObjectTrace.h
@@ -5,7 +5,7 @@
 #include "CoreTypes.h"
 #include "Trace/Config.h"
 
-#if !IS_PROGRAM && !UE_BUILD_SHIPPING
+#if UE_TRACE_ENABLED && !IS_PROGRAM && !UE_BUILD_SHIPPING
 #define OBJECT_TRACE_ENABLED 1
 #else
 #define OBJECT_TRACE_ENABLED 0
diff --git a/Engine/Source/Runtime/TraceLog/Public/Trace/Detail/Trace.h b/Engine/Source/Runtime/TraceLog/Public/Trace/Detail/Trace.h
--- a/Engine/Source/Runtime/TraceLog/Public/Trace/Detail/Trace.h
+++ b/Engine/Source/Runtime/TraceLog/Public/Trace/Detail/Trace.h
@@ -105,7 +105,7 @@
 
 #define TRACE_PRIVATE_CHANNEL(ChannelName)
 
-#define TRACE_PRIVATE_CHANNEL_EXTERN(ModuleApi, ChannelName)
+#define TRACE_PRIVATE_CHANNEL_EXTERN(ChannelName)
 
 #define TRACE_PRIVATE_CHANNEL_DEFINE(ChannelName)
 
