The initialisation of thread-local variables is unsequenced with regard to other translation units which meant TlsContext might not be initialised in time. So now the thread-local buffer pointer and the "null" buffer are constant-initialised variables which will definitely have suitable values before any code is run. TlsContext is still required to detect when threads exit so we don't leak buffers so it's used to get a thread ID to guarantee TlsContext exists and is initialised (compilers are free to lazy-initialise if they so wish so it must be used at least once). Long message this isn't it? No one will read this far.

#rnx

diff --git a/Engine/Source/Runtime/TraceLog/Private/Trace/Writer.cpp b/Engine/Source/Runtime/TraceLog/Private/Trace/Writer.cpp
--- a/Engine/Source/Runtime/TraceLog/Private/Trace/Writer.cpp
+++ b/Engine/Source/Runtime/TraceLog/Private/Trace/Writer.cpp
@@ -62,67 +62,50 @@ void Writer_InitializeTiming()
 
 
 ////////////////////////////////////////////////////////////////////////////////
-static bool GInitialized = false;
+static bool						GInitialized;		// = false;
+static FWriteBuffer				GNullWriteBuffer	= { 0, 0, nullptr, nullptr, (uint8*)&GNullWriteBuffer };
+thread_local FWriteBuffer*		GTlsWriteBuffer		= &GNullWriteBuffer;
+TRACELOG_API uint32 volatile	GLogSerial;			// = 0;
+
 
-////////////////////////////////////////////////////////////////////////////////
-thread_local					FWriteTlsContext TlsContext;
-uint8							FWriteTlsContext::DefaultBuffer[];	// = {}
-uint32 volatile					FWriteTlsContext::ThreadIdCounter;	// = 0;
-TRACELOG_API uint32 volatile	GLogSerial;							// = 0;
 
 ////////////////////////////////////////////////////////////////////////////////
-FWriteTlsContext::FWriteTlsContext()
+struct FWriteTlsContext
 {
-	auto* Target = (FWriteBuffer*)DefaultBuffer;
-
-	static bool Once;
-	if (!Once)
-	{
-		Target->Cursor = DefaultBuffer;
-		Target->ThreadId = 0;
-		Once = true;
-	}
+				~FWriteTlsContext();
+	uint32		GetThreadId();
 
-	Buffer = Target;
-}
+private:
+	uint32		ThreadId = 0;
+};
 
 ////////////////////////////////////////////////////////////////////////////////
 FWriteTlsContext::~FWriteTlsContext()
 {
-	if (GInitialized && HasValidBuffer())
+	if (GInitialized && GTlsWriteBuffer != &GNullWriteBuffer)
 	{
-		UPTRINT EtxOffset = UPTRINT((uint8*)Buffer - Buffer->Cursor);
-		AtomicStoreRelaxed(&(Buffer->EtxOffset), EtxOffset);
+		UPTRINT EtxOffset = UPTRINT((uint8*)GTlsWriteBuffer - GTlsWriteBuffer->Cursor);
+		AtomicStoreRelaxed(&(GTlsWriteBuffer->EtxOffset), EtxOffset);
 	}
 }
 
 ////////////////////////////////////////////////////////////////////////////////
-bool FWriteTlsContext::HasValidBuffer() const
+uint32 FWriteTlsContext::GetThreadId()
 {
-	return (Buffer->ThreadId != 0);
+	if (ThreadId)
+	{
+		return ThreadId;
+	}
+
+	static uint32 volatile Counter;
+	ThreadId = AtomicIncrementRelaxed(&Counter) + 1;
+	return ThreadId;
 }
 
 ////////////////////////////////////////////////////////////////////////////////
-inline FWriteBuffer* FWriteTlsContext::SetBuffer(FWriteBuffer* InBuffer)
-{
-	FWriteBuffer* PrevBuffer = Buffer;
-
-	uint32 ThreadId;
-	if (!PrevBuffer->ThreadId)
-	{
-		ThreadId = AtomicIncrementRelaxed(&ThreadIdCounter) + 1;
-		PrevBuffer = nullptr;
-	}
-	else
-	{
-		ThreadId = PrevBuffer->ThreadId;
-	}
+thread_local FWriteTlsContext	GTlsContext;
 
-	Buffer = InBuffer;
-	Buffer->ThreadId = ThreadId;
 
-	return PrevBuffer;
-}
 
 ////////////////////////////////////////////////////////////////////////////////
 #define T_ALIGN alignas(PLATFORM_CACHE_LINE_SIZE)
@@ -142,7 +125,7 @@ TRACELOG_API FWriteBuffer* Writer_GetBuffer()
 {
 	// Thread locals and DLLs don't mix so for modular builds we are forced to
 	// export this function to access thread-local variables.
-	return TlsContext.GetBuffer();
+	return GTlsWriteBuffer;
 }
 #endif
 
@@ -221,9 +204,11 @@ static FWriteBuffer* Writer_NextBufferInternal(uint32 PageGrowth)
 	NextBuffer->NextBuffer = nullptr;
 
 
-	FWriteBuffer* PrevBuffer = TlsContext.SetBuffer(NextBuffer);
-	if (PrevBuffer == nullptr)
+	FWriteBuffer* CurrentBuffer = GTlsWriteBuffer;
+	if (CurrentBuffer == &GNullWriteBuffer)
 	{
+		NextBuffer->ThreadId = GTlsContext.GetThreadId();
+
 		// Add this next buffer to the active list.
 		for (;; Private::PlatformYield())
 		{
@@ -236,13 +221,15 @@ static FWriteBuffer* Writer_NextBufferInternal(uint32 PageGrowth)
 	}
 	else
 	{
-		PrevBuffer->NextBuffer = NextBuffer;
+		CurrentBuffer->NextBuffer = NextBuffer;
+		NextBuffer->ThreadId = CurrentBuffer->ThreadId;
 
 		// Retire current buffer.
-		UPTRINT EtxOffset = UPTRINT((uint8*)(PrevBuffer) - PrevBuffer->Cursor);
-		AtomicStoreRelease(&(PrevBuffer->EtxOffset), EtxOffset);
+		UPTRINT EtxOffset = UPTRINT((uint8*)(CurrentBuffer) - CurrentBuffer->Cursor);
+		AtomicStoreRelease(&(CurrentBuffer->EtxOffset), EtxOffset);
 	}
 
+	GTlsWriteBuffer = NextBuffer;
 	return NextBuffer;
 }
 
@@ -255,11 +242,10 @@ TRACELOG_API FWriteBuffer* Writer_NextBuffer(int32 Size)
 		return nullptr;
 	}
 
-	// Retire current buffer unless its the initial boot one.
-	if (TlsContext.HasValidBuffer())
+	FWriteBuffer* CurrentBuffer = GTlsWriteBuffer;
+	if (CurrentBuffer != &GNullWriteBuffer)
 	{
-		FWriteBuffer* Current = TlsContext.GetBuffer();
-		Current->Cursor -= Size;
+		CurrentBuffer->Cursor -= Size;
 	}
 
 	FWriteBuffer* NextBuffer = Writer_NextBufferInternal(GPoolPageGrowth);
diff --git a/Engine/Source/Runtime/TraceLog/Public/Trace/Detail/Writer.inl b/Engine/Source/Runtime/TraceLog/Public/Trace/Detail/Writer.inl
--- a/Engine/Source/Runtime/TraceLog/Public/Trace/Detail/Writer.inl
+++ b/Engine/Source/Runtime/TraceLog/Public/Trace/Detail/Writer.inl
@@ -28,23 +28,6 @@ struct FWriteBuffer
 
 
 
-////////////////////////////////////////////////////////////////////////////////
-struct FWriteTlsContext
-{
-							FWriteTlsContext();
-							~FWriteTlsContext();
-	bool					HasValidBuffer() const;
-	FWriteBuffer*			SetBuffer(FWriteBuffer*);
-	FWriteBuffer*			GetBuffer() const { return Buffer; }
-
-private:
-	FWriteBuffer*			Buffer;
-	static uint8			DefaultBuffer[sizeof(FWriteBuffer)];
-	static uint32 volatile	ThreadIdCounter;
-};
-
-
-
 ////////////////////////////////////////////////////////////////////////////////
 extern TRACELOG_API uint32 volatile	GLogSerial;
 TRACELOG_API FWriteBuffer*			Writer_NextBuffer(int32);
@@ -52,10 +35,10 @@ TRACELOG_API FWriteBuffer*			Writer_GetBuffer();
 
 ////////////////////////////////////////////////////////////////////////////////
 #if IS_MONOLITHIC
-extern thread_local FWriteTlsContext TlsContext;
+extern thread_local FWriteBuffer* GTlsWriteBuffer;
 inline FWriteBuffer* Writer_GetBuffer()
 {
-	return TlsContext.GetBuffer();
+	return GTlsWriteBuffer;
 }
 #endif // IS_MONOLITHIC
 
