Do not allow sockets to be inherited. Child processes spawned with UE's CreateProc() are set to inherit their parent's handles. This results in listening sockets lingering after the ower's terminated if its children still a reference to these sockets (i.e.  they were inherited). As these lingering sockets can still be connected to, no end of odd socket behaviour ensues.

#rnx

diff --git a/Engine/Source/Runtime/TraceLog/Private/Trace/Detail/Windows/WindowsTrace.cpp b/Engine/Source/Runtime/TraceLog/Private/Trace/Detail/Windows/WindowsTrace.cpp
--- a/Engine/Source/Runtime/TraceLog/Private/Trace/Detail/Windows/WindowsTrace.cpp
+++ b/Engine/Source/Runtime/TraceLog/Private/Trace/Detail/Windows/WindowsTrace.cpp
@@ -149,7 +149,7 @@ UPTRINT TcpSocketConnect(const ANSICHAR* Host, uint16 Port)
 	// socket() will create a socket with overlapped IO support which we don't
 	// want as it complicates sharing Io*() API with FileOpen(). So we use
 	// WSASocket instead which affords us more control over socket properties
-	SOCKET Socket = WSASocket(AF_INET, SOCK_STREAM, IPPROTO_TCP, nullptr, 0, 0);
+	SOCKET Socket = WSASocket(AF_INET, SOCK_STREAM, IPPROTO_TCP, nullptr, 0, WSA_FLAG_NO_HANDLE_INHERIT);
 	if (Socket == INVALID_SOCKET)
 	{
 		return 0;
@@ -177,7 +177,7 @@ UPTRINT TcpSocketListen(uint16 Port)
 	TcpSocketInitialize();
 
 	// See TcpSocketConnect() for why WSASocket() is used here.
-	SOCKET Socket = WSASocket(AF_INET, SOCK_STREAM, IPPROTO_TCP, nullptr, 0, 0);
+	SOCKET Socket = WSASocket(AF_INET, SOCK_STREAM, IPPROTO_TCP, nullptr, 0, WSA_FLAG_NO_HANDLE_INHERIT);
 	if (Socket == INVALID_SOCKET)
 	{
 		return 0;
