Fixed HoloLens arm64 compile error.

#rnx

diff --git a/Engine/Source/Runtime/TraceLog/Public/Trace/Detail/Atomic.h b/Engine/Source/Runtime/TraceLog/Public/Trace/Detail/Atomic.h
--- a/Engine/Source/Runtime/TraceLog/Public/Trace/Detail/Atomic.h
+++ b/Engine/Source/Runtime/TraceLog/Public/Trace/Detail/Atomic.h
@@ -56,7 +56,11 @@ inline void PlatformYield()
 #if PLATFORM_CPU_X86_FAMILY
 	_mm_pause();
 #elif PLATFORM_CPU_ARM_FAMILY
-	__builtin_arm_yield();
+#	if IS_MSVC
+		__yield();
+#	else
+		__builtin_arm_yield();
+#	endif
 #else
 	#error Unsupported architecture!
 #endif
@@ -65,7 +69,7 @@ inline void PlatformYield()
 ////////////////////////////////////////////////////////////////////////////////
 #if IS_MSVC
 #	if defined(_M_ARM) || defined(_M_ARM64)
-#		define INTERLOCKED_API(Name, Suffix, ...)	Name##_##Suffix(__VA_ARGS__)
+#		define INTERLOCKED_API(Name, Suffix, ...)	Name##Suffix(__VA_ARGS__)
 #	elif defined(_M_IX86) || defined(_M_X64)
 #		define INTERLOCKED_API(Name, Suffix, ...)	Name(__VA_ARGS__)
 #	endif
