platform extensions can enable tracelog by default (it can still be overriden in Target.GlobalDefinitions). probably should become a project setting at some point!

diff --git a/Engine/Source/Runtime/TraceLog/TraceLog.Build.cs b/Engine/Source/Runtime/TraceLog/TraceLog.Build.cs
--- a/Engine/Source/Runtime/TraceLog/TraceLog.Build.cs
+++ b/Engine/Source/Runtime/TraceLog/TraceLog.Build.cs
@@ -11,4 +11,21 @@ public class TraceLog : ModuleRules
 		bRequiresImplementModule = false;
 		PublicIncludePathModuleNames.Add("Core");
     }
+
+	// used by platform extension derived classes. probably should become a project setting at some point!
+	protected void EnableTraceByDefault(ReadOnlyTargetRules Target)
+	{
+		if (Target.Configuration != UnrealTargetConfiguration.Shipping && Target.Type != TargetType.Program)
+		{
+			foreach (String Definition in Target.GlobalDefinitions)
+			{
+				if (Definition.Contains("UE_TRACE_ENABLED"))
+				{
+					// Define already set in Target.GlobalDefinitions
+					return;
+				}
+			}
+			PublicDefinitions.Add("UE_TRACE_ENABLED=1");
+		}
+	}
 }
