Condition order of new threads prior to sending out their events.

#rnx

diff --git a/Engine/Source/Runtime/TraceLog/Private/Trace/Writer.cpp b/Engine/Source/Runtime/TraceLog/Private/Trace/Writer.cpp
--- a/Engine/Source/Runtime/TraceLog/Private/Trace/Writer.cpp
+++ b/Engine/Source/Runtime/TraceLog/Private/Trace/Writer.cpp
@@ -513,6 +513,20 @@ static void Writer_ConsumeEvents()
 		}
 	}
 
+	// Reverse the new threads list so they're more closely ordered by age
+	// when sent out.
+	FWriteBuffer* __restrict NewThreadCursor = NewThreadList;
+	NewThreadList = nullptr;
+	while (NewThreadCursor != nullptr)
+	{
+		FWriteBuffer* __restrict NextThread = NewThreadCursor->NextThread;
+
+		NewThreadCursor->NextThread = NewThreadList;
+		NewThreadList = NewThreadCursor;
+
+		NewThreadCursor = NextThread;
+	}
+
 	FRetireList RetireList;
 
 	FWriteBuffer* __restrict ActiveThreadList = GActiveThreadList;
