Implementation of Trace's important events. Some events are only traced once. CPU scope names or log format strings for example. Those trace-once events get lost if a user then connects to trace late (or reconnects a second time). To combat that important events are traced to a separate stream internally so that they can be identified, cached off to one side, and resent when a connection is first established.

#rnx

diff --git a/Engine/Plugins/Animation/GameplayInsights/Source/GameplayInsights/Private/AnimationAnalyzer.cpp b/Engine/Plugins/Animation/GameplayInsights/Source/GameplayInsights/Private/AnimationAnalyzer.cpp
--- a/Engine/Plugins/Animation/GameplayInsights/Source/GameplayInsights/Private/AnimationAnalyzer.cpp
+++ b/Engine/Plugins/Animation/GameplayInsights/Source/GameplayInsights/Private/AnimationAnalyzer.cpp
@@ -2,6 +2,7 @@
 
 #include "AnimationAnalyzer.h"
 #include "TraceServices/Model/AnalysisSession.h"
+#include "TraceServices/Utils.h"
 #include "AnimationProvider.h"
 #include "Containers/ArrayView.h"
 
@@ -118,7 +119,8 @@ bool FAnimationAnalyzer::OnEvent(uint16 RouteId, EStyle Style, const FOnEventCon
 	case RouteId_Name:
 	{
 		uint32 Id = EventData.GetValue<uint32>("Id");
-		AnimationProvider.AppendName(Id, reinterpret_cast<const TCHAR*>(EventData.GetAttachment()));
+		FString Name = FTraceAnalyzerUtils::LegacyAttachmentString<TCHAR>("Name", Context);
+		AnimationProvider.AppendName(Id, *Name);
 		break;
 	}
 	case RouteId_SkeletalMeshFrame:
diff --git a/Engine/Plugins/Animation/GameplayInsights/Source/GameplayInsights/Private/ObjectPropertyTrace.cpp b/Engine/Plugins/Animation/GameplayInsights/Source/GameplayInsights/Private/ObjectPropertyTrace.cpp
--- a/Engine/Plugins/Animation/GameplayInsights/Source/GameplayInsights/Private/ObjectPropertyTrace.cpp
+++ b/Engine/Plugins/Animation/GameplayInsights/Source/GameplayInsights/Private/ObjectPropertyTrace.cpp
@@ -10,12 +10,12 @@
 #if OBJECT_PROPERTY_TRACE_ENABLED
 UE_TRACE_CHANNEL(ObjectProperties)
 
-UE_TRACE_EVENT_BEGIN(Object, ClassPropertyStringId, Important)
+UE_TRACE_EVENT_BEGIN(Object, ClassPropertyStringId, NoSync|Important)
 	UE_TRACE_EVENT_FIELD(uint32, Id)
 	UE_TRACE_EVENT_FIELD(Trace::WideString, Value)
 UE_TRACE_EVENT_END()
 
-UE_TRACE_EVENT_BEGIN(Object, ClassProperty, Important)
+UE_TRACE_EVENT_BEGIN(Object, ClassProperty, NoSync|Important)
 	UE_TRACE_EVENT_FIELD(uint64, ClassId)
 	UE_TRACE_EVENT_FIELD(int32, Id)
 	UE_TRACE_EVENT_FIELD(int32, ParentId)
@@ -61,9 +61,9 @@ namespace ObjectPropertyTrace
 
 		uint32 NewId = StringIdMap.Add(InString, ++CurrentClassPropertyStringId);
 
-		UE_TRACE_LOG(Object, ClassPropertyStringId, ObjectProperties)
+		UE_TRACE_LOG(Object, ClassPropertyStringId, ObjectProperties, InString.Len() * sizeof(TCHAR))
 			<< ClassPropertyStringId.Id(NewId)
-			<< ClassPropertyStringId.Value(*InString);
+			<< ClassPropertyStringId.Value(*InString, InString.Len());
 
 		return NewId;
 	}
diff --git a/Engine/Source/Developer/TraceAnalysis/Private/Analysis/Engine.cpp b/Engine/Source/Developer/TraceAnalysis/Private/Analysis/Engine.cpp
--- a/Engine/Source/Developer/TraceAnalysis/Private/Analysis/Engine.cpp
+++ b/Engine/Source/Developer/TraceAnalysis/Private/Analysis/Engine.cpp
@@ -938,6 +938,8 @@ void FAnalysisEngine::OnNewTrace(const FOnEventContext& Context)
 	Serial.Value = Hint|0x80000000;
 
 	UserUidBias = EventData.GetValue<uint32>("UserUidBias", uint32(Protocol3::EKnownEventUids::User));
+
+	OnTiming(Context);
 }
 
 ////////////////////////////////////////////////////////////////////////////////
@@ -1000,6 +1002,10 @@ void FAnalysisEngine::OnThreadInfoInternal(const FOnEventContext& Context)
 		}
 	}
 
+	// Implicit ThreadId detail is lost as these events are "important". Fake it
+	uint32 PrevThreadId = ThreadInfo->ThreadId;
+	ThreadInfo->ThreadId = EventData.GetValue<uint32>("ThreadId", PrevThreadId);
+
 	const auto* OuterInfo = (FThreadInfo*)ThreadInfo;
 	for (uint16 i = 0, n = Analyzers.Num(); i < n; ++i)
 	{
@@ -1008,6 +1014,8 @@ void FAnalysisEngine::OnThreadInfoInternal(const FOnEventContext& Context)
 			Analyzer->OnThreadInfo(*OuterInfo);
 		}
 	}
+
+	ThreadInfo->ThreadId = PrevThreadId;
 }
 
 ////////////////////////////////////////////////////////////////////////////////
@@ -1834,12 +1842,28 @@ int32 FAnalysisEngine::OnDataProtocol4Impl(
 	FAuxDataCollector AuxCollector;
 	if (Dispatch->Flags & FDispatch::Flag_MaybeHasAux)
 	{
+		// Important events' size may include their array data so we need to backtrack
+		auto NextMark = Reader.SaveMark();
+		if (Dispatch->Flags & FDispatch::Flag_Important)
+		{
+			Reader.RestoreMark(Mark);
+			Reader.Advance(sizeof(FEventHeader) + Dispatch->EventSize);
+		}
+
 		int AuxStatus = OnDataProtocol2Aux(Reader, AuxCollector);
 		if (AuxStatus == 0)
 		{
 			Reader.RestoreMark(Mark);
 			return 1;
 		}
+
+		// User error could have resulted in less space being used that was
+		// allocated for important events. So we can't assume that aux data
+		// reading has read all the way up to the next event. So we use marks
+		if (Dispatch->Flags & FDispatch::Flag_Important)
+		{
+			Reader.RestoreMark(NextMark);
+		}
 	}
 
 	// Maintain sync
diff --git a/Engine/Source/Developer/TraceServices/Private/Analyzers/CountersTraceAnalysis.cpp b/Engine/Source/Developer/TraceServices/Private/Analyzers/CountersTraceAnalysis.cpp
--- a/Engine/Source/Developer/TraceServices/Private/Analyzers/CountersTraceAnalysis.cpp
+++ b/Engine/Source/Developer/TraceServices/Private/Analyzers/CountersTraceAnalysis.cpp
@@ -1,5 +1,6 @@
 // Copyright Epic Games, Inc. All Rights Reserved.
 #include "CountersTraceAnalysis.h"
+#include "Common/Utils.h"
 #include "TraceServices/Model/Counters.h"
 #include "ProfilingDebugging/CountersTrace.h"
 
@@ -40,7 +41,9 @@ bool FCountersAnalyzer::OnEvent(uint16 RouteId, EStyle Style, const FOnEventCont
 		{
 			Counter->SetDisplayHint(Trace::CounterDisplayHint_Memory);
 		}
-		Counter->SetName(Session.StoreString(reinterpret_cast<const TCHAR*>(EventData.GetAttachment())));
+		FString Name = FTraceAnalyzerUtils::LegacyAttachmentString<TCHAR>("Name", Context);
+		const TCHAR* NamePtr = Session.StoreString(*Name);
+		Counter->SetName(NamePtr);
 		CountersMap.Add(CounterId, Counter);
 		break;
 	}
diff --git a/Engine/Source/Developer/TraceServices/Private/Analyzers/CpuProfilerTraceAnalysis.cpp b/Engine/Source/Developer/TraceServices/Private/Analyzers/CpuProfilerTraceAnalysis.cpp
--- a/Engine/Source/Developer/TraceServices/Private/Analyzers/CpuProfilerTraceAnalysis.cpp
+++ b/Engine/Source/Developer/TraceServices/Private/Analyzers/CpuProfilerTraceAnalysis.cpp
@@ -42,14 +42,24 @@ bool FCpuProfilerAnalyzer::OnEvent(uint16 RouteId, EStyle Style, const FOnEventC
 	case RouteId_EventSpec:
 	{
 		uint32 SpecId = EventData.GetValue<uint32>("Id");
-		uint8 CharSize = EventData.GetValue<uint8>("CharSize");
-		if (CharSize == sizeof(ANSICHAR))
+		
+		FString Name;
+		if (EventData.GetString("Name", Name))
 		{
-			DefineScope(SpecId, Session.StoreString(StringCast<TCHAR>(reinterpret_cast<const ANSICHAR*>(EventData.GetAttachment())).Get()));
+			const TCHAR* NamePtr = Session.StoreString(*Name);
+			DefineScope(SpecId, NamePtr);
 		}
-		else if (CharSize == 0 || CharSize == sizeof(TCHAR)) // 0 for backwards compatibility
+		else
 		{
-			DefineScope(SpecId, Session.StoreString(reinterpret_cast<const TCHAR*>(EventData.GetAttachment())));
+			uint8 CharSize = EventData.GetValue<uint8>("CharSize");
+			if (CharSize == sizeof(ANSICHAR))
+			{
+				DefineScope(SpecId, Session.StoreString(StringCast<TCHAR>(reinterpret_cast<const ANSICHAR*>(EventData.GetAttachment())).Get()));
+			}
+			else if (CharSize == 0 || CharSize == sizeof(TCHAR)) // 0 for backwards compatibility
+			{
+				DefineScope(SpecId, Session.StoreString(reinterpret_cast<const TCHAR*>(EventData.GetAttachment())));
+			}
 		}
 		break;
 	}
diff --git a/Engine/Source/Developer/TraceServices/Private/Analyzers/CsvProfilerTraceAnalysis.cpp b/Engine/Source/Developer/TraceServices/Private/Analyzers/CsvProfilerTraceAnalysis.cpp
--- a/Engine/Source/Developer/TraceServices/Private/Analyzers/CsvProfilerTraceAnalysis.cpp
+++ b/Engine/Source/Developer/TraceServices/Private/Analyzers/CsvProfilerTraceAnalysis.cpp
@@ -70,22 +70,24 @@ bool FCsvProfilerAnalyzer::OnEvent(uint16 RouteId, EStyle Style, const FOnEventC
 	case RouteId_RegisterCategory:
 	{
 		int32 CategoryIndex = EventData.GetValue<int32>("Index");
-		const TCHAR* Name = reinterpret_cast<const TCHAR*>(EventData.GetAttachment());
-		CategoryMap.Add(CategoryIndex, Session.StoreString(Name));;
+		FString Name = FTraceAnalyzerUtils::LegacyAttachmentString<TCHAR>("Name", Context);
+		CategoryMap.Add(CategoryIndex, Session.StoreString(*Name));;
 		break;
 	}
 	case RouteId_DefineInlineStat:
 	{
 		uint64 StatId = EventData.GetValue<uint64>("StatId");
 		int32 CategoryIndex = EventData.GetValue<int32>("CategoryIndex");
-		DefineStatSeries(StatId, ANSI_TO_TCHAR(reinterpret_cast<const ANSICHAR*>(EventData.GetAttachment())), CategoryIndex, true);
+		FString Name = FTraceAnalyzerUtils::LegacyAttachmentString<ANSICHAR>("Name", Context);
+		DefineStatSeries(StatId, *Name, CategoryIndex, true);
 		break;
 	}
 	case RouteId_DefineDeclaredStat:
 	{
 		uint64 StatId = EventData.GetValue<uint64>("StatId");
 		int32 CategoryIndex = EventData.GetValue<int32>("CategoryIndex");
-		DefineStatSeries(StatId, reinterpret_cast<const TCHAR*>(EventData.GetAttachment()), CategoryIndex, false);
+		FString Name = FTraceAnalyzerUtils::LegacyAttachmentString<TCHAR>("Name", Context);
+		DefineStatSeries(StatId, *Name, CategoryIndex, false);
 		break;
 	}
 	case RouteId_BeginStat:
diff --git a/Engine/Source/Developer/TraceServices/Private/Analyzers/LoadTimeTraceAnalysis.cpp b/Engine/Source/Developer/TraceServices/Private/Analyzers/LoadTimeTraceAnalysis.cpp
--- a/Engine/Source/Developer/TraceServices/Private/Analyzers/LoadTimeTraceAnalysis.cpp
+++ b/Engine/Source/Developer/TraceServices/Private/Analyzers/LoadTimeTraceAnalysis.cpp
@@ -380,7 +380,8 @@ bool FAsyncLoadingTraceAnalyzer::OnEvent(uint16 RouteId, EStyle Style, const FOn
 	{
 		Trace::FAnalysisSessionEditScope _(Session);
 		uint64 ClassPtr = EventData.GetValue<uint64>("Class");
-		const Trace::FClassInfo& ClassInfo = LoadTimeProfilerProvider.AddClassInfo(reinterpret_cast<const TCHAR*>(EventData.GetAttachment()));
+		FString Name = FTraceAnalyzerUtils::LegacyAttachmentString<TCHAR>("Name", Context);
+		const Trace::FClassInfo& ClassInfo = LoadTimeProfilerProvider.AddClassInfo(*Name);
 		ClassInfosMap.Add(ClassPtr, &ClassInfo);
 		break;
 	}
diff --git a/Engine/Source/Developer/TraceServices/Private/Analyzers/LogTraceAnalysis.cpp b/Engine/Source/Developer/TraceServices/Private/Analyzers/LogTraceAnalysis.cpp
--- a/Engine/Source/Developer/TraceServices/Private/Analyzers/LogTraceAnalysis.cpp
+++ b/Engine/Source/Developer/TraceServices/Private/Analyzers/LogTraceAnalysis.cpp
@@ -1,6 +1,7 @@
 // Copyright Epic Games, Inc. All Rights Reserved.
 #include "LogTraceAnalysis.h"
 #include "AnalysisServicePrivate.h"
+#include "Common/Utils.h"
 #include "Logging/LogTrace.h"
 #include "Model/LogPrivate.h"
 
@@ -31,7 +32,10 @@ bool FLogTraceAnalyzer::OnEvent(uint16 RouteId, EStyle Style, const FOnEventCont
 	{
 		uint64 CategoryPointer = EventData.GetValue<uint64>("CategoryPointer");
 		Trace::FLogCategory& Category = LogProvider.GetCategory(CategoryPointer);
-		Category.Name = Session.StoreString(reinterpret_cast<const TCHAR*>(EventData.GetAttachment()));
+
+		FString Name = FTraceAnalyzerUtils::LegacyAttachmentString<TCHAR>("Name", Context);
+		Category.Name = Session.StoreString(*Name);
+
 		Category.DefaultVerbosity = static_cast<ELogVerbosity::Type>(EventData.GetValue<uint8>("DefaultVerbosity"));
 		break;
 	}
@@ -44,9 +48,23 @@ bool FLogTraceAnalyzer::OnEvent(uint16 RouteId, EStyle Style, const FOnEventCont
 		Spec.Category = &Category;
 		Spec.Line = EventData.GetValue<int32>("Line");
 		Spec.Verbosity = static_cast<ELogVerbosity::Type>(EventData.GetValue<uint8>("Verbosity"));
-		const ANSICHAR* File = reinterpret_cast<const ANSICHAR*>(EventData.GetAttachment());
-		Spec.File = Session.StoreString(ANSI_TO_TCHAR(File));
-		Spec.FormatString = Session.StoreString(reinterpret_cast<const TCHAR*>(EventData.GetAttachment() + strlen(File) + 1));
+
+		FString FileName;
+		if (EventData.GetString("FileName", FileName))
+		{
+			Spec.File = Session.StoreString(*FileName);
+
+			FString FormatString;
+			EventData.GetString("FormatString", FormatString);
+			Spec.FormatString = Session.StoreString(*FormatString);
+		}
+		else
+		{
+			const ANSICHAR* File = reinterpret_cast<const ANSICHAR*>(EventData.GetAttachment());
+			Spec.File = Session.StoreString(ANSI_TO_TCHAR(File));
+			Spec.FormatString = Session.StoreString(reinterpret_cast<const TCHAR*>(EventData.GetAttachment() + strlen(File) + 1));
+		}
+
 		break;
 	}
 	case RouteId_LogMessage:
diff --git a/Engine/Source/Developer/TraceServices/Private/Analyzers/MiscTraceAnalysis.cpp b/Engine/Source/Developer/TraceServices/Private/Analyzers/MiscTraceAnalysis.cpp
--- a/Engine/Source/Developer/TraceServices/Private/Analyzers/MiscTraceAnalysis.cpp
+++ b/Engine/Source/Developer/TraceServices/Private/Analyzers/MiscTraceAnalysis.cpp
@@ -76,9 +76,23 @@ bool FMiscTraceAnalyzer::OnEvent(uint16 RouteId, EStyle Style, const FOnEventCon
 		uint64 BookmarkPoint = EventData.GetValue<uint64>("BookmarkPoint");
 		Trace::FBookmarkSpec& Spec = BookmarkProvider.GetSpec(BookmarkPoint);
 		Spec.Line = EventData.GetValue<int32>("Line");
-		const ANSICHAR* File = reinterpret_cast<const ANSICHAR*>(EventData.GetAttachment());
-		Spec.File = Session.StoreString(ANSI_TO_TCHAR(File));
-		Spec.FormatString = Session.StoreString(reinterpret_cast<const TCHAR*>(EventData.GetAttachment() + strlen(File) + 1));
+
+		FString FileName;
+		if (EventData.GetString("FileName", FileName))
+		{
+			Spec.File = Session.StoreString(*FileName);
+
+			FString FormatString;
+			EventData.GetString("FormatString", FormatString);
+			Spec.FormatString = Session.StoreString(*FormatString);
+		}
+		else
+		{
+			const ANSICHAR* File = reinterpret_cast<const ANSICHAR*>(EventData.GetAttachment());
+			Spec.File = Session.StoreString(ANSI_TO_TCHAR(File));
+			Spec.FormatString = Session.StoreString(reinterpret_cast<const TCHAR*>(EventData.GetAttachment() + strlen(File) + 1));
+		}
+
 
 		Trace::FLogMessageSpec& LogMessageSpec = LogProvider.GetMessageSpec(BookmarkPoint);
 		LogMessageSpec.Category = &LogProvider.GetCategory(Trace::FLogProvider::ReservedLogCategory_Bookmark);
diff --git a/Engine/Source/Developer/TraceServices/Private/Common/Utils.h b/Engine/Source/Developer/TraceServices/Private/Common/Utils.h
--- a/Engine/Source/Developer/TraceServices/Private/Common/Utils.h
+++ b/Engine/Source/Developer/TraceServices/Private/Common/Utils.h
@@ -4,55 +4,43 @@
 
 #include "CoreTypes.h"
 #include "Trace/Analyzer.h"
+#include "TraceServices/Utils.h"
 
-struct FTraceAnalyzerUtils
+namespace FTraceAnalyzerUtils
 {
-	static uint64 Decode7bit(const uint8*& BufferPtr)
-	{
-		uint64 Value = 0;
-		uint64 ByteIndex = 0;
-		bool HasMoreBytes;
-		do
-		{
-			uint8 ByteValue = *BufferPtr++;
-			HasMoreBytes = ByteValue & 0x80;
-			Value |= uint64(ByteValue & 0x7f) << (ByteIndex * 7);
-			++ByteIndex;
-		} while (HasMoreBytes);
-		return Value;
-	}
 
-	static int64 DecodeZigZag(const uint8*& BufferPtr)
+inline uint64 Decode7bit(const uint8*& BufferPtr)
+{
+	uint64 Value = 0;
+	uint64 ByteIndex = 0;
+	bool HasMoreBytes;
+	do
 	{
-		uint64 Z = Decode7bit(BufferPtr);
-		return (Z & 1) ? (Z >> 1) ^ -1 : (Z >> 1);
-	}
+		uint8 ByteValue = *BufferPtr++;
+		HasMoreBytes = ByteValue & 0x80;
+		Value |= uint64(ByteValue & 0x7f) << (ByteIndex * 7);
+		++ByteIndex;
+	} while (HasMoreBytes);
+	return Value;
+}
 
-	static uint32 GetThreadIdField(
-		const Trace::IAnalyzer::FOnEventContext& Context,
-		const ANSICHAR* FieldName="ThreadId")
-	{
-		// Trace analysis was changed to be able to provide a suitable id. Prior to
-		// this users of Trace would send along their own thread ids. For backwards
-		// compatibility we'll bias field thread ids to avoid collision with Trace's.
-		static const uint32 Bias = 0x70000000;
-		uint32 ThreadId = Context.EventData.GetValue<uint32>(FieldName, 0);
-		ThreadId |= ThreadId ? Bias : Context.ThreadInfo.GetId();
-		return ThreadId;
-	}
+inline int64 DecodeZigZag(const uint8*& BufferPtr)
+{
+	uint64 Z = Decode7bit(BufferPtr);
+	return (Z & 1) ? (Z >> 1) ^ -1 : (Z >> 1);
+}
 
-	template <typename AttachedCharType>
-	static FString LegacyAttachmentString(
-		const ANSICHAR* FieldName,
-		const Trace::IAnalyzer::FOnEventContext& Context)
-	{
-		FString Out;
-		if (!Context.EventData.GetString(FieldName, Out))
-		{
-			Out = FString(
-					Context.EventData.GetAttachmentSize() / sizeof(AttachedCharType),
-					(const AttachedCharType*)(Context.EventData.GetAttachment()));
-		}
-		return Out;
-	}
-};
+inline uint32 GetThreadIdField(
+	const Trace::IAnalyzer::FOnEventContext& Context,
+	const ANSICHAR* FieldName="ThreadId")
+{
+	// Trace analysis was changed to be able to provide a suitable id. Prior to
+	// this users of Trace would send along their own thread ids. For backwards
+	// compatibility we'll bias field thread ids to avoid collision with Trace's.
+	static const uint32 Bias = 0x70000000;
+	uint32 ThreadId = Context.EventData.GetValue<uint32>(FieldName, 0);
+	ThreadId |= ThreadId ? Bias : Context.ThreadInfo.GetId();
+	return ThreadId;
+}
+
+} // namespace FTraceAnalyzerUtils
diff --git a/Engine/Source/Developer/TraceServices/Private/Model/Channel.cpp b/Engine/Source/Developer/TraceServices/Private/Model/Channel.cpp
--- a/Engine/Source/Developer/TraceServices/Private/Model/Channel.cpp
+++ b/Engine/Source/Developer/TraceServices/Private/Model/Channel.cpp
@@ -17,7 +17,12 @@ FChannelProvider::FChannelProvider()
 void FChannelProvider::AnnounceChannel(const TCHAR* InChannelName, uint32 Id, bool bReadOnly)
 {
 	FString ChannelName(InChannelName);
-	ChannelName.GetCharArray()[0] = TChar<TCHAR>::ToUpper(ChannelName.GetCharArray()[0]);
+
+	if (*InChannelName)
+	{
+		ChannelName.GetCharArray()[0] = TChar<TCHAR>::ToUpper(ChannelName.GetCharArray()[0]);
+	}
+
 	Channels.Add(FChannelEntry{
 		Id,
 		ChannelName,
diff --git a/Engine/Source/Developer/TraceServices/Public/TraceServices/Utils.h b/Engine/Source/Developer/TraceServices/Public/TraceServices/Utils.h
new file mode 100644
--- /dev/null
+++ b/Engine/Source/Developer/TraceServices/Public/TraceServices/Utils.h
@@ -0,0 +1,23 @@
+// Copyright Epic Games, Inc. All Rights Reserved.
+
+#pragma once
+
+namespace FTraceAnalyzerUtils
+{
+
+template <typename AttachedCharType>
+inline FString LegacyAttachmentString(
+	const ANSICHAR* FieldName,
+	const Trace::IAnalyzer::FOnEventContext& Context)
+{
+	FString Out;
+	if (!Context.EventData.GetString(FieldName, Out))
+	{
+		Out = FString(
+				Context.EventData.GetAttachmentSize() / sizeof(AttachedCharType),
+				(const AttachedCharType*)(Context.EventData.GetAttachment()));
+	}
+	return Out;
+}
+
+} // namespace FTraceAnalyzerUtils
diff --git a/Engine/Source/Runtime/Core/Private/HAL/LowLevelMemTracker.cpp b/Engine/Source/Runtime/Core/Private/HAL/LowLevelMemTracker.cpp
--- a/Engine/Source/Runtime/Core/Private/HAL/LowLevelMemTracker.cpp
+++ b/Engine/Source/Runtime/Core/Private/HAL/LowLevelMemTracker.cpp
@@ -21,13 +21,13 @@
 
 UE_TRACE_CHANNEL(MemoryChannel, "Memory overview", true)
 
-UE_TRACE_EVENT_BEGIN(LLM, TagsSpec, Important)
+UE_TRACE_EVENT_BEGIN(LLM, TagsSpec, NoSync|Important)
 	UE_TRACE_EVENT_FIELD(const void*, TagId)
 	UE_TRACE_EVENT_FIELD(const void*, ParentId)
 	UE_TRACE_EVENT_FIELD(Trace::AnsiString, Name)
 UE_TRACE_EVENT_END()
 
-UE_TRACE_EVENT_BEGIN(LLM, TrackerSpec, Important)
+UE_TRACE_EVENT_BEGIN(LLM, TrackerSpec, NoSync|Important)
 	UE_TRACE_EVENT_FIELD(uint8, TrackerId)
 	UE_TRACE_EVENT_FIELD(Trace::AnsiString, Name)
 UE_TRACE_EVENT_END()
@@ -3701,11 +3701,18 @@ namespace LLMPrivate
 		if (!bTrackerSpecSent)
 		{
 			bTrackerSpecSent = true;
-			static const ANSICHAR* TrackerNames[] = {"Platform", "Default"};
+			static struct {
+				const ANSICHAR* Name;
+				uint32 Len;
+			} TrackerNames[] = {
+				{ "Platform", 8 },
+				{ "Default", 7 },
+			};
 			static_assert(UE_ARRAY_COUNT(TrackerNames) == int(ELLMTracker::Max), "");
-			UE_TRACE_LOG(LLM, TrackerSpec, MemoryChannel)
+			uint32 NameLen = TrackerNames[(uint8)Tracker].Len;
+			UE_TRACE_LOG(LLM, TrackerSpec, MemoryChannel, NameLen * sizeof(ANSICHAR))
 				<< TrackerSpec.TrackerId((uint8)Tracker)
-				<< TrackerSpec.Name(TrackerNames[(uint8)Tracker]);
+				<< TrackerSpec.Name(TrackerNames[(uint8)Tracker].Name, NameLen);
 		}
 
 		TStringBuilder<1024> NameBuffer;
@@ -3720,7 +3727,7 @@ namespace LLMPrivate
 			const FTagData* Parent = TagData->GetParent();
 			NameBuffer.Reset();
 			TagData->AppendDisplayPath(NameBuffer);
-			UE_TRACE_LOG(LLM, TagsSpec, MemoryChannel)
+			UE_TRACE_LOG(LLM, TagsSpec, MemoryChannel, NameBuffer.Len() * sizeof(ANSICHAR))
 				<< TagsSpec.TagId(GetTagId(TagData))
 				<< TagsSpec.ParentId(GetTagId(Parent))
 				<< TagsSpec.Name(*NameBuffer, NameBuffer.Len());
@@ -3792,10 +3799,10 @@ namespace LLMPrivate
 		{
 			const FTagData* TagData = LLMRef.FindTagData(static_cast<ELLMTag>(GenericTagIndex));
 			FString TagName = TagData->GetDisplayPath();
-			UE_TRACE_LOG(LLM, TagsSpec, MemoryChannel)
+			UE_TRACE_LOG(LLM, TagsSpec, MemoryChannel, TagName.Len() * sizeof(ANSICHAR))
 				<< TagsSpec.TagId(GetTagId(TagData))
 				<< TagsSpec.ParentId(GetTagId(TagData->GetParent()))
-				<< TagsSpec.Name(*TagName);
+				<< TagsSpec.Name(*TagName, TagName.Len());
 		}
 	}
 }
diff --git a/Engine/Source/Runtime/Core/Private/Logging/LogTrace.cpp b/Engine/Source/Runtime/Core/Private/Logging/LogTrace.cpp
--- a/Engine/Source/Runtime/Core/Private/Logging/LogTrace.cpp
+++ b/Engine/Source/Runtime/Core/Private/Logging/LogTrace.cpp
@@ -10,16 +10,19 @@
 
 UE_TRACE_CHANNEL(LogChannel)
 
-UE_TRACE_EVENT_BEGIN(Logging, LogCategory, Important)
+UE_TRACE_EVENT_BEGIN(Logging, LogCategory, NoSync|Important)
 	UE_TRACE_EVENT_FIELD(const void*, CategoryPointer)
 	UE_TRACE_EVENT_FIELD(uint8, DefaultVerbosity)
+	UE_TRACE_EVENT_FIELD(Trace::AnsiString, Name)
 UE_TRACE_EVENT_END()
 
-UE_TRACE_EVENT_BEGIN(Logging, LogMessageSpec, Important)
+UE_TRACE_EVENT_BEGIN(Logging, LogMessageSpec, NoSync|Important)
 	UE_TRACE_EVENT_FIELD(const void*, LogPoint)
 	UE_TRACE_EVENT_FIELD(const void*, CategoryPointer)
 	UE_TRACE_EVENT_FIELD(int32, Line)
 	UE_TRACE_EVENT_FIELD(uint8, Verbosity)
+	UE_TRACE_EVENT_FIELD(Trace::AnsiString, FileName)
+	UE_TRACE_EVENT_FIELD(Trace::WideString, FormatString)
 UE_TRACE_EVENT_END()
 
 UE_TRACE_EVENT_BEGIN(Logging, LogMessage)
@@ -29,27 +32,25 @@ UE_TRACE_EVENT_END()
 
 void FLogTrace::OutputLogCategory(const FLogCategoryBase* Category, const TCHAR* Name, ELogVerbosity::Type DefaultVerbosity)
 {
-	uint16 NameSize = (uint16)((FCString::Strlen(Name) + 1) * sizeof(TCHAR));
-	UE_TRACE_LOG(Logging, LogCategory, LogChannel, NameSize)
+	uint16 NameLen = uint16(FCString::Strlen(Name));
+	UE_TRACE_LOG(Logging, LogCategory, LogChannel, NameLen * sizeof(ANSICHAR))
 		<< LogCategory.CategoryPointer(Category)
 		<< LogCategory.DefaultVerbosity(DefaultVerbosity)
-		<< LogCategory.Attachment(Name, NameSize);
+		<< LogCategory.Name(Name, NameLen);
 }
 
 void FLogTrace::OutputLogMessageSpec(const void* LogPoint, const FLogCategoryBase* Category, ELogVerbosity::Type Verbosity, const ANSICHAR* File, int32 Line, const TCHAR* Format)
 {
-	uint16 FileNameSize = (uint16)(strlen(File) + 1);
-	uint16 FormatStringSize = (uint16)((FCString::Strlen(Format) + 1) * sizeof(TCHAR));
-	auto StringCopyFunc = [FileNameSize, FormatStringSize, File, Format](uint8* Out) {
-		memcpy(Out, File, FileNameSize);
-		memcpy(Out + FileNameSize, Format, FormatStringSize);
-	};
-	UE_TRACE_LOG(Logging, LogMessageSpec, LogChannel, FileNameSize + FormatStringSize)
+	uint16 FileNameLen = uint16(strlen(File));
+	uint16 FormatStringLen = uint16(FCString::Strlen(Format));
+	uint32 DataSize = (FileNameLen * sizeof(ANSICHAR)) + (FormatStringLen * sizeof(TCHAR));
+	UE_TRACE_LOG(Logging, LogMessageSpec, LogChannel, DataSize)
 		<< LogMessageSpec.LogPoint(LogPoint)
 		<< LogMessageSpec.CategoryPointer(Category)
 		<< LogMessageSpec.Line(Line)
 		<< LogMessageSpec.Verbosity(Verbosity)
-		<< LogMessageSpec.Attachment(StringCopyFunc);
+		<< LogMessageSpec.FileName(File, FileNameLen)
+		<< LogMessageSpec.FormatString(Format, FormatStringLen);
 }
 
 void FLogTrace::OutputLogMessageInternal(const void* LogPoint, uint16 EncodedFormatArgsSize, uint8* EncodedFormatArgs)
diff --git a/Engine/Source/Runtime/Core/Private/ProfilingDebugging/CountersTrace.cpp b/Engine/Source/Runtime/Core/Private/ProfilingDebugging/CountersTrace.cpp
--- a/Engine/Source/Runtime/Core/Private/ProfilingDebugging/CountersTrace.cpp
+++ b/Engine/Source/Runtime/Core/Private/ProfilingDebugging/CountersTrace.cpp
@@ -8,10 +8,11 @@
 
 UE_TRACE_CHANNEL_DEFINE(CountersChannel)
 
-UE_TRACE_EVENT_BEGIN(Counters, Spec, Important)
+UE_TRACE_EVENT_BEGIN(Counters, Spec, NoSync|Important)
 	UE_TRACE_EVENT_FIELD(uint16, Id)
 	UE_TRACE_EVENT_FIELD(uint8, Type)
 	UE_TRACE_EVENT_FIELD(uint8, DisplayHint)
+	UE_TRACE_EVENT_FIELD(Trace::AnsiString, Name)
 UE_TRACE_EVENT_END()
 
 UE_TRACE_EVENT_BEGIN(Counters, SetValueInt)
@@ -29,22 +30,31 @@ UE_TRACE_EVENT_END()
 
 uint16 FCountersTrace::OutputInitCounter(const TCHAR* CounterName, ETraceCounterType CounterType, ETraceCounterDisplayHint CounterDisplayHint)
 {
-	static TAtomic<uint16> NextId(1);
 	if (!UE_TRACE_CHANNELEXPR_IS_ENABLED(CountersChannel))
+	{
 		return 0;
-	uint16 CounterId = uint16(NextId++);
-	uint16 NameSize = (uint16)((FCString::Strlen(CounterName) + 1) * sizeof(TCHAR));
-	UE_TRACE_LOG(Counters, Spec, CountersChannel, NameSize)
+	}
+
+	static TAtomic<uint16> NextId;
+
+	uint16 CounterId = uint16(NextId++) + 1;
+	uint16 NameLen = uint16(FCString::Strlen(CounterName));
+	UE_TRACE_LOG(Counters, Spec, CountersChannel, NameLen * sizeof(ANSICHAR))
 		<< Spec.Id(CounterId)
 		<< Spec.Type(uint8(CounterType))
 		<< Spec.DisplayHint(uint8(CounterDisplayHint))
-		<< Spec.Attachment(CounterName, NameSize);
+		<< Spec.Name(CounterName, NameLen);
 	return CounterId;
 }
 
 void FCountersTrace::OutputSetValue(uint16 CounterId, int64 Value)
 {
-	UE_TRACE_LOG(Counters, SetValueInt, CountersChannel && CounterId)
+	if (!CounterId)
+	{
+		return;
+	}
+
+	UE_TRACE_LOG(Counters, SetValueInt, CountersChannel)
 		<< SetValueInt.Cycle(FPlatformTime::Cycles64())
 		<< SetValueInt.Value(Value)
 		<< SetValueInt.CounterId(CounterId);
@@ -52,7 +62,12 @@ void FCountersTrace::OutputSetValue(uint16 CounterId, int64 Value)
 
 void FCountersTrace::OutputSetValue(uint16 CounterId, double Value)
 {
-	UE_TRACE_LOG(Counters, SetValueFloat, CountersChannel && CounterId)
+	if (!CounterId)
+	{
+		return;
+	}
+
+	UE_TRACE_LOG(Counters, SetValueFloat, CountersChannel)
 		<< SetValueFloat.Cycle(FPlatformTime::Cycles64())
 		<< SetValueFloat.Value(Value)
 		<< SetValueFloat.CounterId(CounterId);
diff --git a/Engine/Source/Runtime/Core/Private/ProfilingDebugging/CpuProfilerTrace.cpp b/Engine/Source/Runtime/Core/Private/ProfilingDebugging/CpuProfilerTrace.cpp
--- a/Engine/Source/Runtime/Core/Private/ProfilingDebugging/CpuProfilerTrace.cpp
+++ b/Engine/Source/Runtime/Core/Private/ProfilingDebugging/CpuProfilerTrace.cpp
@@ -14,15 +14,15 @@
 
 UE_TRACE_CHANNEL_DEFINE(CpuChannel)
 
-UE_TRACE_EVENT_BEGIN(CpuProfiler, EventSpec, Important)
+UE_TRACE_EVENT_BEGIN(CpuProfiler, EventSpec, NoSync|Important)
 	UE_TRACE_EVENT_FIELD(uint32, Id)
-	UE_TRACE_EVENT_FIELD(uint8, CharSize)
+	UE_TRACE_EVENT_FIELD(Trace::AnsiString, Name)
 UE_TRACE_EVENT_END()
 
 UE_TRACE_EVENT_BEGIN(CpuProfiler, EventBatch, NoSync)
 UE_TRACE_EVENT_END()
 
-UE_TRACE_EVENT_BEGIN(CpuProfiler, EndCapture, Important)
+UE_TRACE_EVENT_BEGIN(CpuProfiler, EndCapture)
 UE_TRACE_EVENT_END()
 
 UE_TRACE_EVENT_BEGIN(CpuProfiler, EndThread, NoSync)
@@ -200,29 +200,27 @@ void FCpuProfilerTrace::OutputEndEvent()
 
 uint32 FCpuProfilerTraceInternal::GetNextSpecId()
 {
-	static TAtomic<uint32> NextSpecId(1);
-	return NextSpecId++;
+	static TAtomic<uint32> NextSpecId;
+	return (NextSpecId++) + 1;
 }
 
 uint32 FCpuProfilerTrace::OutputEventType(const TCHAR* Name)
 {
 	uint32 SpecId = FCpuProfilerTraceInternal::GetNextSpecId();
-	uint16 NameSize = (uint16)((FCString::Strlen(Name) + 1) * sizeof(TCHAR));
-	UE_TRACE_LOG(CpuProfiler, EventSpec, CpuChannel, NameSize)
+	uint16 NameLen = uint16(FCString::Strlen(Name));
+	UE_TRACE_LOG(CpuProfiler, EventSpec, CpuChannel, NameLen * sizeof(ANSICHAR))
 		<< EventSpec.Id(SpecId)
-		<< EventSpec.CharSize(uint8(sizeof(TCHAR)))
-		<< EventSpec.Attachment(Name, NameSize);
+		<< EventSpec.Name(Name, NameLen);
 	return SpecId;
 }
 
 uint32 FCpuProfilerTrace::OutputEventType(const ANSICHAR* Name)
 {
 	uint32 SpecId = FCpuProfilerTraceInternal::GetNextSpecId();
-	uint16 NameSize = (uint16)(strlen(Name) + 1);
-	UE_TRACE_LOG(CpuProfiler, EventSpec, CpuChannel, NameSize)
+	uint16 NameLen = uint16(strlen(Name));
+	UE_TRACE_LOG(CpuProfiler, EventSpec, CpuChannel, NameLen * sizeof(ANSICHAR))
 		<< EventSpec.Id(SpecId)
-		<< EventSpec.CharSize(uint8(1))
-		<< EventSpec.Attachment(Name, NameSize);
+		<< EventSpec.Name(Name, NameLen);
 	return SpecId;
 }
 
diff --git a/Engine/Source/Runtime/Core/Private/ProfilingDebugging/CsvProfilerTrace.cpp b/Engine/Source/Runtime/Core/Private/ProfilingDebugging/CsvProfilerTrace.cpp
--- a/Engine/Source/Runtime/Core/Private/ProfilingDebugging/CsvProfilerTrace.cpp
+++ b/Engine/Source/Runtime/Core/Private/ProfilingDebugging/CsvProfilerTrace.cpp
@@ -8,18 +8,21 @@
 
 #if CSVPROFILERTRACE_ENABLED
 
-UE_TRACE_EVENT_BEGIN(CsvProfiler, RegisterCategory, Important)
+UE_TRACE_EVENT_BEGIN(CsvProfiler, RegisterCategory, NoSync|Important)
 	UE_TRACE_EVENT_FIELD(int32, Index)
+	UE_TRACE_EVENT_FIELD(Trace::AnsiString, Name)
 UE_TRACE_EVENT_END()
 
-UE_TRACE_EVENT_BEGIN(CsvProfiler, DefineInlineStat, Important)
+UE_TRACE_EVENT_BEGIN(CsvProfiler, DefineInlineStat, NoSync|Important)
 	UE_TRACE_EVENT_FIELD(uint64, StatId)
 	UE_TRACE_EVENT_FIELD(int32, CategoryIndex)
+	UE_TRACE_EVENT_FIELD(Trace::AnsiString, Name)
 UE_TRACE_EVENT_END()
 
-UE_TRACE_EVENT_BEGIN(CsvProfiler, DefineDeclaredStat, Important)
+UE_TRACE_EVENT_BEGIN(CsvProfiler, DefineDeclaredStat, NoSync|Important)
 	UE_TRACE_EVENT_FIELD(uint64, StatId)
 	UE_TRACE_EVENT_FIELD(int32, CategoryIndex)
+	UE_TRACE_EVENT_FIELD(Trace::AnsiString, Name)
 UE_TRACE_EVENT_END()
 
 UE_TRACE_EVENT_BEGIN(CsvProfiler, BeginStat)
@@ -110,19 +113,19 @@ struct FCsvProfilerTraceInternal
 
 void FCsvProfilerTrace::OutputRegisterCategory(int32 Index, const TCHAR* Name)
 {
-	uint16 NameSize = (uint16)((FCString::Strlen(Name) + 1) * sizeof(TCHAR));
-	UE_TRACE_LOG(CsvProfiler, RegisterCategory, CountersChannel, NameSize)
+	uint16 NameLen = uint16(FCString::Strlen(Name));
+	UE_TRACE_LOG(CsvProfiler, RegisterCategory, CountersChannel, NameLen * sizeof(ANSICHAR))
 		<< RegisterCategory.Index(Index)
-		<< RegisterCategory.Attachment(Name, NameSize);
+		<< RegisterCategory.Name(Name, NameLen);
 }
 
 void FCsvProfilerTrace::OutputInlineStat(const char* StatName, int32 CategoryIndex)
 {
-	uint16 NameSize = (uint16)((strlen(StatName) + 1) * sizeof(char));
-	UE_TRACE_LOG(CsvProfiler, DefineInlineStat, CountersChannel, NameSize)
+	uint16 NameLen = uint16(strlen(StatName));
+	UE_TRACE_LOG(CsvProfiler, DefineInlineStat, CountersChannel, NameLen * sizeof(ANSICHAR))
 		<< DefineInlineStat.StatId(FCsvProfilerTraceInternal::GetStatId(StatName, CategoryIndex))
 		<< DefineInlineStat.CategoryIndex(CategoryIndex)
-		<< DefineInlineStat.Attachment(StatName, NameSize);
+		<< DefineInlineStat.Name(StatName, NameLen);
 }
 
 CSV_DECLARE_CATEGORY_EXTERN(Exclusive);
@@ -136,11 +139,11 @@ void FCsvProfilerTrace::OutputDeclaredStat(const FName& StatName, int32 Category
 {
 	TCHAR NameString[NAME_SIZE];
 	StatName.GetPlainNameString(NameString);
-	uint16 NameSize = (uint16)((FCString::Strlen(NameString) + 1) * sizeof(TCHAR));
-	UE_TRACE_LOG(CsvProfiler, DefineDeclaredStat, CountersChannel, NameSize)
+	uint16 NameLen = uint16(FCString::Strlen(NameString));
+	UE_TRACE_LOG(CsvProfiler, DefineDeclaredStat, CountersChannel, NameLen * sizeof(ANSICHAR))
 		<< DefineDeclaredStat.StatId(FCsvProfilerTraceInternal::GetStatId(StatName, CategoryIndex))
 		<< DefineDeclaredStat.CategoryIndex(CategoryIndex)
-		<< DefineDeclaredStat.Attachment(NameString, NameSize);
+		<< DefineDeclaredStat.Name(NameString, NameLen);
 }
 
 void FCsvProfilerTrace::OutputBeginStat(const char* StatName, int32 CategoryIndex, uint64 Cycles)
diff --git a/Engine/Source/Runtime/Core/Private/ProfilingDebugging/MiscTrace.cpp b/Engine/Source/Runtime/Core/Private/ProfilingDebugging/MiscTrace.cpp
--- a/Engine/Source/Runtime/Core/Private/ProfilingDebugging/MiscTrace.cpp
+++ b/Engine/Source/Runtime/Core/Private/ProfilingDebugging/MiscTrace.cpp
@@ -11,9 +11,11 @@
 UE_TRACE_CHANNEL(FrameChannel)
 UE_TRACE_CHANNEL(BookmarkChannel)
 
-UE_TRACE_EVENT_BEGIN(Misc, BookmarkSpec, Important)
+UE_TRACE_EVENT_BEGIN(Misc, BookmarkSpec, NoSync|Important)
 	UE_TRACE_EVENT_FIELD(const void*, BookmarkPoint)
 	UE_TRACE_EVENT_FIELD(int32, Line)
+	UE_TRACE_EVENT_FIELD(Trace::WideString, FormatString)
+	UE_TRACE_EVENT_FIELD(Trace::AnsiString, FileName)
 UE_TRACE_EVENT_END()
 
 UE_TRACE_EVENT_BEGIN(Misc, Bookmark)
@@ -43,16 +45,15 @@ uint64 FMiscTraceInternal::LastFrameCycle[TraceFrameType_Count] = { 0, 0 };
 
 void FMiscTrace::OutputBookmarkSpec(const void* BookmarkPoint, const ANSICHAR* File, int32 Line, const TCHAR* Format)
 {
-	uint16 FileNameSize = (uint16)(strlen(File) + 1);
-	uint16 FormatStringSize = (uint16)((FCString::Strlen(Format) + 1) * sizeof(TCHAR));
-	auto StringCopyFunc = [FileNameSize, FormatStringSize, File, Format](uint8* Out) {
-		memcpy(Out, File, FileNameSize);
-		memcpy(Out + FileNameSize, Format, FormatStringSize);
-	};
-	UE_TRACE_LOG(Misc, BookmarkSpec, BookmarkChannel, FileNameSize + FormatStringSize)
+	uint16 FileNameLen = uint16(strlen(File));
+	uint16 FormatStringLen = uint16(FCString::Strlen(Format));
+
+	uint32 DataSize = (FileNameLen * sizeof(ANSICHAR)) + (FormatStringLen * sizeof(TCHAR));
+	UE_TRACE_LOG(Misc, BookmarkSpec, BookmarkChannel, DataSize)
 		<< BookmarkSpec.BookmarkPoint(BookmarkPoint)
 		<< BookmarkSpec.Line(Line)
-		<< BookmarkSpec.Attachment(StringCopyFunc);
+		<< BookmarkSpec.FormatString(Format, FormatStringLen)
+		<< BookmarkSpec.FileName(File, FileNameLen);
 }
 
 void FMiscTrace::OutputBookmarkInternal(const void* BookmarkPoint, uint16 EncodedFormatArgsSize, uint8* EncodedFormatArgs)
diff --git a/Engine/Source/Runtime/Core/Private/ProfilingDebugging/TraceAuxiliary.cpp b/Engine/Source/Runtime/Core/Private/ProfilingDebugging/TraceAuxiliary.cpp
--- a/Engine/Source/Runtime/Core/Private/ProfilingDebugging/TraceAuxiliary.cpp
+++ b/Engine/Source/Runtime/Core/Private/ProfilingDebugging/TraceAuxiliary.cpp
@@ -48,15 +48,15 @@ public:
 	template <class T> void	ReadChannels(T&& Callback) const;
 	void					AddChannels(const TCHAR* ChannelList);
 	bool					Connect(ETraceConnectType Type, const TCHAR* Parameter);
+	bool					Stop();
 	void					EnableChannels();
 	void					DisableChannels();
 
 private:
 	enum class EState : uint8
 	{
-		None,
-		Tracing,
 		Stopped,
+		Tracing,
 	};
 
 	struct FChannel
@@ -73,7 +73,7 @@ private:
 
 	TMap<uint32, FChannel>	Channels;
 	FString					TraceDest;
-	EState					State = EState::None;
+	EState					State = EState::Stopped;
 };
 
 static FTraceAuxiliaryImpl GTraceAuxiliary;
@@ -173,6 +173,19 @@ bool FTraceAuxiliaryImpl::Connect(ETraceConnectType Type, const TCHAR* Parameter
 	return true;
 }
 
+////////////////////////////////////////////////////////////////////////////////
+bool FTraceAuxiliaryImpl::Stop()
+{
+	if (!Trace::Stop())
+	{
+		return false;
+	}
+
+	DisableChannels();
+	State = EState::Stopped;
+	return true;
+}
+
 ////////////////////////////////////////////////////////////////////////////////
 void FTraceAuxiliaryImpl::EnableChannel(FChannel& Channel)
 {
@@ -335,8 +348,8 @@ static void TraceAuxiliaryStart(const TArray<FString>& Args)
 ////////////////////////////////////////////////////////////////////////////////
 static void TraceAuxiliaryStop()
 {
-	GTraceAuxiliary.DisableChannels();
-	UE_LOG(LogConsoleResponse, Log, TEXT("Tracing stopped. Use 'Trace.Start' to resume"));
+	GTraceAuxiliary.Stop();
+	UE_LOG(LogConsoleResponse, Log, TEXT("Tracing stopped."));
 }
 
 ////////////////////////////////////////////////////////////////////////////////
@@ -361,7 +374,7 @@ static FAutoConsoleCommand TraceAuxiliaryStopCmd(
 
 
 ////////////////////////////////////////////////////////////////////////////////
-UE_TRACE_EVENT_BEGIN(Diagnostics, Session2, Important)
+UE_TRACE_EVENT_BEGIN(Diagnostics, Session2, NoSync|Important)
 	UE_TRACE_EVENT_FIELD(Trace::AnsiString, Platform)
 	UE_TRACE_EVENT_FIELD(Trace::AnsiString, AppName)
 	UE_TRACE_EVENT_FIELD(Trace::WideString, CommandLine)
@@ -375,7 +388,11 @@ void FTraceAuxiliary::Initialize(const TCHAR* CommandLine)
 #if UE_TRACE_ENABLED
 	// Trace out information about this session. This is done before initialisation
 	// so that it is always sent (all channels are enabled prior to initialisation)
-	UE_TRACE_LOG(Diagnostics, Session2, Trace::TraceLogChannel)
+	uint32 DataSize = 
+		(UE_ARRAY_COUNT(PREPROCESSOR_TO_STRING(UBT_COMPILED_PLATFORM)) * sizeof(ANSICHAR)) +
+		(UE_ARRAY_COUNT(UE_APP_NAME) * sizeof(ANSICHAR)) +
+		(FCString::Strlen(CommandLine) * sizeof(TCHAR));
+	UE_TRACE_LOG(Diagnostics, Session2, Trace::TraceLogChannel, DataSize)
 		<< Session2.Platform(PREPROCESSOR_TO_STRING(UBT_COMPILED_PLATFORM))
 		<< Session2.AppName(UE_APP_NAME)
 		<< Session2.CommandLine(CommandLine)
diff --git a/Engine/Source/Runtime/Core/Private/Stats/StatsTrace.cpp b/Engine/Source/Runtime/Core/Private/Stats/StatsTrace.cpp
--- a/Engine/Source/Runtime/Core/Private/Stats/StatsTrace.cpp
+++ b/Engine/Source/Runtime/Core/Private/Stats/StatsTrace.cpp
@@ -12,7 +12,7 @@
 
 UE_TRACE_CHANNEL(StatsChannel)
 
-UE_TRACE_EVENT_BEGIN(Stats, Spec, Important)
+UE_TRACE_EVENT_BEGIN(Stats, Spec, NoSync|Important)
 	UE_TRACE_EVENT_FIELD(uint32, Id)
 	UE_TRACE_EVENT_FIELD(bool, IsFloatingPoint)
 	UE_TRACE_EVENT_FIELD(bool, IsMemory)
diff --git a/Engine/Source/Runtime/CoreUObject/Private/Serialization/LoadTimeTrace.cpp b/Engine/Source/Runtime/CoreUObject/Private/Serialization/LoadTimeTrace.cpp
--- a/Engine/Source/Runtime/CoreUObject/Private/Serialization/LoadTimeTrace.cpp
+++ b/Engine/Source/Runtime/CoreUObject/Private/Serialization/LoadTimeTrace.cpp
@@ -14,7 +14,7 @@
 #include "UObject/UObjectGlobals.h"
 #include "Misc/CommandLine.h"
 
-UE_TRACE_EVENT_BEGIN(LoadTime, StartAsyncLoading, Important)
+UE_TRACE_EVENT_BEGIN(LoadTime, StartAsyncLoading, NoSync|Important)
 	UE_TRACE_EVENT_FIELD(uint64, Cycle)
 UE_TRACE_EVENT_END()
 
@@ -101,8 +101,9 @@ UE_TRACE_EVENT_BEGIN(LoadTime, AsyncPackageImportDependency)
 	UE_TRACE_EVENT_FIELD(const void*, ImportedAsyncPackage)
 UE_TRACE_EVENT_END()
 
-UE_TRACE_EVENT_BEGIN(LoadTime, ClassInfo, Important)
+UE_TRACE_EVENT_BEGIN(LoadTime, ClassInfo, NoSync|Important)
 	UE_TRACE_EVENT_FIELD(const UClass*, Class)
+	UE_TRACE_EVENT_FIELD(Trace::AnsiString, Name)
 UE_TRACE_EVENT_END()
 
 void FLoadTimeProfilerTracePrivate::Init()
@@ -198,18 +199,18 @@ void FLoadTimeProfilerTracePrivate::OutputAsyncPackageImportDependency(const voi
 void FLoadTimeProfilerTracePrivate::OutputClassInfo(const UClass* Class, const FName& Name)
 {
 	TCHAR Buffer[FName::StringBufferSize];
-	uint16 NameSize = (Name.ToString(Buffer) + 1) * sizeof(TCHAR);
-	UE_TRACE_LOG(LoadTime, ClassInfo, LoadTimeChannel, NameSize)
+	uint16 NameLen = Name.ToString(Buffer);
+	UE_TRACE_LOG(LoadTime, ClassInfo, LoadTimeChannel, NameLen * sizeof(ANSICHAR))
 		<< ClassInfo.Class(Class)
-		<< ClassInfo.Attachment(Buffer, NameSize);
+		<< ClassInfo.Name(Buffer, NameLen);
 }
 
 void FLoadTimeProfilerTracePrivate::OutputClassInfo(const UClass* Class, const TCHAR* Name)
 {
-	uint16 NameSize = (FCString::Strlen(Name) + 1) * sizeof(TCHAR);
-	UE_TRACE_LOG(LoadTime, ClassInfo, LoadTimeChannel, NameSize)
+	uint16 NameLen = uint16(FCString::Strlen(Name));
+	UE_TRACE_LOG(LoadTime, ClassInfo, LoadTimeChannel, NameLen * sizeof(ANSICHAR))
 		<< ClassInfo.Class(Class)
-		<< ClassInfo.Attachment(Name, NameSize);
+		<< ClassInfo.Name(Name, NameLen);
 }
 
 FLoadTimeProfilerTracePrivate::FCreateExportScope::FCreateExportScope(const void* AsyncPackage, const UObject* const* InObject)
diff --git a/Engine/Source/Runtime/Engine/Private/Animation/AnimTrace.cpp b/Engine/Source/Runtime/Engine/Private/Animation/AnimTrace.cpp
--- a/Engine/Source/Runtime/Engine/Private/Animation/AnimTrace.cpp
+++ b/Engine/Source/Runtime/Engine/Private/Animation/AnimTrace.cpp
@@ -39,7 +39,7 @@ UE_TRACE_EVENT_BEGIN(Animation, TickRecord)
 	UE_TRACE_EVENT_FIELD(bool, IsBlendSpace)
 UE_TRACE_EVENT_END()
 
-UE_TRACE_EVENT_BEGIN(Animation, SkeletalMesh2, Important)
+UE_TRACE_EVENT_BEGIN(Animation, SkeletalMesh2, NoSync|Important)
 	UE_TRACE_EVENT_FIELD(uint64, Id)
 	UE_TRACE_EVENT_FIELD(int32[], ParentIndices)
 UE_TRACE_EVENT_END()
@@ -189,8 +189,9 @@ UE_TRACE_EVENT_BEGIN(Animation, StateMachineState)
 	UE_TRACE_EVENT_FIELD(float, ElapsedTime)
 UE_TRACE_EVENT_END()
 
-UE_TRACE_EVENT_BEGIN(Animation, Name, Important)
+UE_TRACE_EVENT_BEGIN(Animation, Name, NoSync|Important)
 	UE_TRACE_EVENT_FIELD(uint32, Id)
+	UE_TRACE_EVENT_FIELD(Trace::WideString, Name)
 UE_TRACE_EVENT_END()
 
 UE_TRACE_EVENT_BEGIN(Animation, Notify)
@@ -433,7 +434,7 @@ void FAnimTrace::OutputSkeletalMesh(const USkeletalMesh* InMesh)
 		ParentIndices[BoneIndex++] = BoneInfo.ParentIndex;
 	}
 
-	UE_TRACE_LOG(Animation, SkeletalMesh2, AnimationChannel)
+	UE_TRACE_LOG(Animation, SkeletalMesh2, AnimationChannel, ParentIndices.Num() * sizeof(int32))
 		<< SkeletalMesh2.Id(FObjectTrace::GetObjectId(InMesh))
 		<< SkeletalMesh2.ParentIndices(ParentIndices.GetData(), ParentIndices.Num());
 
@@ -463,16 +464,12 @@ uint32 FAnimTrace::OutputName(const FName& InName)
 
 	if(bShouldTrace)
 	{
-		int32 NameStringLength = InName.GetStringLength() + 1;
-
-		auto StringCopyFunc = [NameStringLength, &InName](uint8* Out)
-		{
-			InName.ToString(reinterpret_cast<TCHAR*>(Out), NameStringLength);
-		};
+		TCHAR Buffer[256];
+		int32 NameStringLength = InName.ToString(Buffer);
 
 		UE_TRACE_LOG(Animation, Name, AnimationChannel, NameStringLength * sizeof(TCHAR))
 			<< Name.Id(NameId)
-			<< Name.Attachment(StringCopyFunc);
+			<< Name.Name(Buffer, NameStringLength);
 	}
 
 	return NameId;
@@ -1107,4 +1104,4 @@ void FAnimTrace::OutputMontage(UAnimInstance* InAnimInstance, const FAnimMontage
 	}
 }
 
-#endif
\ No newline at end of file
+#endif
diff --git a/Engine/Source/Runtime/Engine/Private/ObjectTrace.cpp b/Engine/Source/Runtime/Engine/Private/ObjectTrace.cpp
--- a/Engine/Source/Runtime/Engine/Private/ObjectTrace.cpp
+++ b/Engine/Source/Runtime/Engine/Private/ObjectTrace.cpp
@@ -21,13 +21,13 @@
 
 UE_TRACE_CHANNEL(ObjectChannel)
 
-UE_TRACE_EVENT_BEGIN(Object, Class, Important)
+UE_TRACE_EVENT_BEGIN(Object, Class)
 	UE_TRACE_EVENT_FIELD(uint64, Id)
 	UE_TRACE_EVENT_FIELD(uint64, SuperId)
 	UE_TRACE_EVENT_FIELD(int32, ClassNameStringLength)
 UE_TRACE_EVENT_END()
 
-UE_TRACE_EVENT_BEGIN(Object, Object, Important)
+UE_TRACE_EVENT_BEGIN(Object, Object)
 	UE_TRACE_EVENT_FIELD(uint64, Id)
 	UE_TRACE_EVENT_FIELD(uint64, ClassId)
 	UE_TRACE_EVENT_FIELD(uint64, OuterId)
@@ -40,7 +40,7 @@ UE_TRACE_EVENT_BEGIN(Object, ObjectEvent)
 	UE_TRACE_EVENT_FIELD(uint8, Event)
 UE_TRACE_EVENT_END()
 
-UE_TRACE_EVENT_BEGIN(Object, World, Important)
+UE_TRACE_EVENT_BEGIN(Object, World)
 	UE_TRACE_EVENT_FIELD(uint64, Id)
 	UE_TRACE_EVENT_FIELD(int32, PIEInstanceId)
 	UE_TRACE_EVENT_FIELD(uint8, Type)
diff --git a/Engine/Source/Runtime/RHI/Private/GpuProfilerTrace.cpp b/Engine/Source/Runtime/RHI/Private/GpuProfilerTrace.cpp
--- a/Engine/Source/Runtime/RHI/Private/GpuProfilerTrace.cpp
+++ b/Engine/Source/Runtime/RHI/Private/GpuProfilerTrace.cpp
@@ -30,10 +30,10 @@ struct
 
 static TSet<uint32> GEventNames;
 
-RHI_API UE_TRACE_CHANNEL_EXTERN(GpuChannel)
+UE_TRACE_CHANNEL_EXTERN(GpuChannel, RHI_API)
 UE_TRACE_CHANNEL_DEFINE(GpuChannel)
 
-UE_TRACE_EVENT_BEGIN(GpuProfiler, EventSpec, Important)
+UE_TRACE_EVENT_BEGIN(GpuProfiler, EventSpec, NoSync|Important)
 	UE_TRACE_EVENT_FIELD(uint32, EventType)
 	UE_TRACE_EVENT_FIELD(uint16[], Name)
 UE_TRACE_EVENT_END()
@@ -84,7 +84,7 @@ void FGpuProfilerTrace::SpecifyEventByName(const FName& Name)
 		uint32 NameLength = String.Len() + 1;
 		static_assert(sizeof(TCHAR) == sizeof(uint16), "");
 
-		UE_TRACE_LOG(GpuProfiler, EventSpec, GpuChannel)
+		UE_TRACE_LOG(GpuProfiler, EventSpec, GpuChannel, NameLength * sizeof(uint16))
 			<< EventSpec.EventType(Index)
 			<< EventSpec.Name((const uint16*)(*String), NameLength);
 	}
diff --git a/Engine/Source/Runtime/TraceLog/Private/Trace/Channel.cpp b/Engine/Source/Runtime/TraceLog/Private/Trace/Channel.cpp
--- a/Engine/Source/Runtime/TraceLog/Private/Trace/Channel.cpp
+++ b/Engine/Source/Runtime/TraceLog/Private/Trace/Channel.cpp
@@ -24,14 +24,14 @@ static FTraceChannel	TraceLogChannelDetail;
 FChannel&				TraceLogChannel			= TraceLogChannelDetail;
 
 ///////////////////////////////////////////////////////////////////////////////
-UE_TRACE_EVENT_BEGIN(Trace, ChannelAnnounce, Important)
+UE_TRACE_EVENT_BEGIN(Trace, ChannelAnnounce, NoSync|Important)
 	UE_TRACE_EVENT_FIELD(uint32, Id)
 	UE_TRACE_EVENT_FIELD(bool, IsEnabled)
 	UE_TRACE_EVENT_FIELD(bool, ReadOnly)
 	UE_TRACE_EVENT_FIELD(Trace::AnsiString, Name)
 UE_TRACE_EVENT_END()
 
-UE_TRACE_EVENT_BEGIN(Trace, ChannelToggle, Important)
+UE_TRACE_EVENT_BEGIN(Trace, ChannelToggle, NoSync|Important)
 	UE_TRACE_EVENT_FIELD(uint32, Id)
 	UE_TRACE_EVENT_FIELD(bool, IsEnabled)
 UE_TRACE_EVENT_END()
@@ -166,7 +166,7 @@ void FChannel::Setup(const ANSICHAR* InChannelName, const InitArgs& InArgs)
 ///////////////////////////////////////////////////////////////////////////////
 void FChannel::Announce() const
 {
-	UE_TRACE_LOG(Trace, ChannelAnnounce, TraceLogChannel)
+	UE_TRACE_LOG(Trace, ChannelAnnounce, TraceLogChannel, Name.Len * sizeof(ANSICHAR))
 		<< ChannelAnnounce.Id(Name.Hash)
 		<< ChannelAnnounce.IsEnabled(IsEnabled())
 		<< ChannelAnnounce.ReadOnly(Args.bReadOnly)
diff --git a/Engine/Source/Runtime/TraceLog/Private/Trace/EventNode.cpp b/Engine/Source/Runtime/TraceLog/Private/Trace/EventNode.cpp
--- a/Engine/Source/Runtime/TraceLog/Private/Trace/EventNode.cpp
+++ b/Engine/Source/Runtime/TraceLog/Private/Trace/EventNode.cpp
@@ -13,16 +13,23 @@ namespace Private {
 ////////////////////////////////////////////////////////////////////////////////
 void					Writer_InternalInitialize();
 FEventNode* volatile	GNewEventList; // = nullptr;
+FEventNode*				GEventListHead;// = nullptr;
+FEventNode*				GEventListTail;// = nullptr;
 
 
 
 ////////////////////////////////////////////////////////////////////////////////
 const FEventNode* FEventNode::FIter::GetNext()
 {
-	auto* Ret = (const FEventNode*)Inner;
+	auto* Ret = (FEventNode*)Inner;
 	if (Ret != nullptr)
 	{
 		Inner = Ret->Next;
+
+		if (Inner == nullptr)
+		{
+			GEventListTail = Ret;
+		}
 	}
 	return Ret;
 }
@@ -44,6 +51,15 @@ FEventNode::FIter FEventNode::ReadNew()
 		EventList = AtomicLoadRelaxed(&GNewEventList);
 	}
 
+	if (GEventListHead == nullptr)
+	{
+		GEventListHead = EventList;
+	}
+	else
+	{
+		GEventListTail->Next = EventList;
+	}
+
 	return { EventList };
 }
 
@@ -112,7 +128,13 @@ void FEventNode::Describe() const
 	EventSize += sizeof(FNewEventEvent::Fields[0]) * Info->FieldCount;
 	EventSize += NamesSize;
 
-	FLogScope LogScope = FLogScope::Enter<FEventInfo::Flag_NoSync>(EventUid, EventSize);
+	struct FNewEventDummy
+	{
+		static constexpr uint32 GetSize()	{ return 0; }
+		static constexpr uint32 GetUid()	{ return 0; }
+		enum { EventFlags = FEventInfo::Flag_NoSync };
+	};
+	FLogScope LogScope = FLogScope::Enter<FNewEventDummy>(EventSize);
 	auto& Event = *(FNewEventEvent*)(LogScope.GetPointer());
 
 	// Write event's main properties.
@@ -157,6 +179,27 @@ void FEventNode::Describe() const
 	LogScope.Commit();
 }
 
+////////////////////////////////////////////////////////////////////////////////
+void FEventNode::OnConnect()
+{
+	// Re-add known events back as new events so that they get described again
+	if (GEventListHead == nullptr)
+	{
+		return;
+	}
+
+	for (;; PlatformYield())
+	{
+		FEventNode* Node = AtomicLoadRelaxed(&GNewEventList);
+		if (AtomicCompareExchangeRelaxed(&GNewEventList, GEventListHead, Node))
+		{
+			break;
+		}
+	}
+
+	GEventListHead = GEventListTail = nullptr;
+}
+
 } // namespace Private
 } // namespace Trace
 
diff --git a/Engine/Source/Runtime/TraceLog/Private/Trace/Field.cpp b/Engine/Source/Runtime/TraceLog/Private/Trace/Field.cpp
--- a/Engine/Source/Runtime/TraceLog/Private/Trace/Field.cpp
+++ b/Engine/Source/Runtime/TraceLog/Private/Trace/Field.cpp
@@ -22,7 +22,7 @@ static void Field_WriteAuxData(uint32 Index, int32 Size, CallbackType&& Callback
 
 	auto* Header = (FAuxHeader*)(Buffer->Cursor - sizeof(FAuxHeader));
 	Header->Size = Size << 8;
-	Header->FieldIndex = uint8(0x80 | (Index & int(EIndexPack::FieldCountMask)));
+	Header->FieldIndex = uint8(0x80 | (Index & int(EIndexPack::NumFieldsMask)));
 
 	bool bCommit = ((uint8*)Header + bMaybeHasAux == Buffer->Committed);
 
diff --git a/Engine/Source/Runtime/TraceLog/Private/Trace/Important/Cache.cpp b/Engine/Source/Runtime/TraceLog/Private/Trace/Important/Cache.cpp
new file mode 100644
--- /dev/null
+++ b/Engine/Source/Runtime/TraceLog/Private/Trace/Important/Cache.cpp
@@ -0,0 +1,179 @@
+// Copyright Epic Games, Inc. All Rights Reserved.
+
+#include "Trace/Config.h"
+
+#if UE_TRACE_ENABLED
+
+#include "Trace/Trace.h"
+
+#include <memory.h>
+
+namespace Trace {
+namespace Private {
+
+////////////////////////////////////////////////////////////////////////////////
+uint32	GetEncodeMaxSize(uint32);
+int32	Encode(const void*, int32, void*, int32);
+void*	Writer_MemoryAllocate(SIZE_T, uint32);
+void	Writer_MemoryFree(void*, uint32);
+void	Writer_SendDataRaw(const void*, uint32);
+uint32	Writer_SendData(uint8* __restrict, uint32);
+
+////////////////////////////////////////////////////////////////////////////////
+struct alignas(16) FCacheBuffer
+{
+	union
+	{
+		FCacheBuffer*	Next;
+		FCacheBuffer**	TailNext;
+	};
+	uint32				Size;
+	uint32				Remaining;
+	uint32				_Unused;
+	uint32				Underflow; // For packet header
+	uint8				Data[];
+};
+
+
+
+////////////////////////////////////////////////////////////////////////////////
+static const uint32		GCacheBufferSize	= 4 << 10;
+static const uint32		GCacheCollectorSize	= 1 << 10;
+static FCacheBuffer*	GCacheCollector;	// = nullptr;
+static FCacheBuffer*	GCacheActiveBuffer;	// = nullptr;
+static FCacheBuffer*	GCacheHeadBuffer;	// = nullptr;
+extern FStatistics		GTraceStatistics;
+
+////////////////////////////////////////////////////////////////////////////////
+static FCacheBuffer* Writer_CacheCreateBuffer(uint32 Size)
+{
+	void* Block = Writer_MemoryAllocate(sizeof(FCacheBuffer) + Size, alignof(FCacheBuffer));
+	auto* Buffer = (FCacheBuffer*)Block;
+	Buffer->Size = Size;
+	Buffer->Remaining = Buffer->Size;
+	Buffer->Next = nullptr;
+	return Buffer;
+}
+
+////////////////////////////////////////////////////////////////////////////////
+static void Writer_CacheCommit(const FCacheBuffer* Collector)
+{
+	struct FPacketEncoded
+	{
+		uint16	PacketSize;
+		uint16	ThreadId;
+		uint16	DecodedSize;
+		uint8	Data[];
+	};
+
+	// Check there's enough size to compress Data into.
+	uint32 InputSize = uint32(Collector->Size - Collector->Remaining);
+	uint32 EncodeMaxSize = GetEncodeMaxSize(InputSize);
+	if (EncodeMaxSize + sizeof(FPacketEncoded) > GCacheActiveBuffer->Remaining)
+	{
+		// Retire active buffer
+		*(GCacheActiveBuffer->TailNext) = GCacheActiveBuffer;
+		GCacheActiveBuffer->TailNext = nullptr;
+
+		// Fire up a new active buffer
+		FCacheBuffer* NewBuffer = Writer_CacheCreateBuffer(GCacheBufferSize);
+		NewBuffer->TailNext = &(GCacheActiveBuffer->Next);
+		GCacheActiveBuffer = NewBuffer;
+
+#if TRACE_PRIVATE_STATISTICS
+		GTraceStatistics.CacheWaste += GCacheActiveBuffer->Remaining;
+#endif
+	}
+
+	uint32 Used = GCacheActiveBuffer->Size - GCacheActiveBuffer->Remaining;
+	auto* Packet = (FPacketEncoded*)(GCacheActiveBuffer->Data + Used);
+	uint32 OutputSize = Encode(Collector->Data, InputSize, Packet->Data, EncodeMaxSize);
+
+	Packet->PacketSize = OutputSize + sizeof(FPacketEncoded);
+	Packet->ThreadId = 0x8000 | uint16(0);
+	Packet->DecodedSize = uint16(InputSize);
+
+	Used = sizeof(FPacketEncoded) + OutputSize;
+	GCacheActiveBuffer->Remaining -= Used;
+
+#if TRACE_PRIVATE_STATISTICS
+	GTraceStatistics.CacheUsed += Used;
+#endif
+}
+
+////////////////////////////////////////////////////////////////////////////////
+void Writer_CacheData(uint8* Data, uint32 Size)
+{
+	Writer_SendData(Data, Size);
+
+	while (true)
+	{
+		uint32 StepSize = (Size < GCacheCollector->Remaining) ? Size : GCacheCollector->Remaining;
+
+		uint32 Used = GCacheCollector->Size - GCacheCollector->Remaining;
+		memcpy(GCacheCollector->Data + Used, Data, StepSize);
+
+		GCacheCollector->Remaining -= StepSize;
+		if (GCacheCollector->Remaining == 0)
+		{
+			Writer_CacheCommit(GCacheCollector);
+			GCacheCollector->Remaining = GCacheCollector->Size;
+		}
+
+		Size -= StepSize;
+		if (Size == 0)
+		{
+			break;
+		}
+
+		Data += StepSize;
+	}
+}
+
+////////////////////////////////////////////////////////////////////////////////
+void Writer_CacheOnConnect()
+{
+	for (FCacheBuffer* Buffer = GCacheHeadBuffer; Buffer != nullptr; Buffer = Buffer->Next)
+	{
+		uint32 Used = Buffer->Size - Buffer->Remaining;
+		Writer_SendDataRaw(Buffer->Data, Used);
+	}
+
+	if (uint32 Used = GCacheActiveBuffer->Size - GCacheActiveBuffer->Remaining)
+	{
+		Writer_SendDataRaw(GCacheActiveBuffer->Data, Used);
+	}
+
+	if (uint32 Used = GCacheCollector->Size - GCacheCollector->Remaining)
+	{
+		Writer_SendData(GCacheCollector->Data, Used);
+	}
+}
+
+////////////////////////////////////////////////////////////////////////////////
+void Writer_InitializeCache()
+{
+	GCacheCollector = Writer_CacheCreateBuffer(GCacheCollectorSize);
+
+	GCacheActiveBuffer = Writer_CacheCreateBuffer(GCacheBufferSize);
+	GCacheActiveBuffer->TailNext = &GCacheHeadBuffer;
+}
+
+////////////////////////////////////////////////////////////////////////////////
+void Writer_ShutdownCache()
+{
+	for (FCacheBuffer* Buffer = GCacheHeadBuffer; Buffer != nullptr;)
+	{
+		FCacheBuffer* Next = Buffer->Next;
+		Writer_MemoryFree(Buffer, GCacheBufferSize);
+		Buffer = Next;
+	}
+
+	Writer_MemoryFree(GCacheActiveBuffer, GCacheBufferSize);
+	Writer_MemoryFree(GCacheCollector, GCacheCollectorSize);
+}
+
+} // namespace Private
+} // namespace Trace
+
+#endif // UE_TRACE_ENABLED
diff --git a/Engine/Source/Runtime/TraceLog/Private/Trace/Important/SharedBuffer.cpp b/Engine/Source/Runtime/TraceLog/Private/Trace/Important/SharedBuffer.cpp
new file mode 100644
--- /dev/null
+++ b/Engine/Source/Runtime/TraceLog/Private/Trace/Important/SharedBuffer.cpp
@@ -0,0 +1,198 @@
+// Copyright Epic Games, Inc. All Rights Reserved.
+
+#include "Trace/Detail/Important/SharedBuffer.h"
+
+#if UE_TRACE_ENABLED
+
+#include "CoreTypes.h"
+#include "Trace/Detail/Atomic.h"
+#include "Trace/Detail/Important/ImportantLogScope.inl"
+
+namespace Trace {
+namespace Private {
+
+////////////////////////////////////////////////////////////////////////////////
+void*	Writer_MemoryAllocate(SIZE_T, uint32);
+void	Writer_MemoryFree(void*, uint32);
+void	Writer_CacheData(uint8*, uint32);
+void	Writer_InitializeCache();
+void	Writer_ShutdownCache();
+
+////////////////////////////////////////////////////////////////////////////////
+static FSharedBuffer	GNullSharedBuffer	= { 0, FSharedBuffer::RefInit };
+FSharedBuffer* volatile GSharedBuffer		= &GNullSharedBuffer;
+static FSharedBuffer*	GTailBuffer;		// = nullptr
+static uint32			GTailPreSent;		// = 0
+static const uint32		GBlockSize			= 1024;//16 << 10;
+extern FStatistics		GTraceStatistics;
+
+////////////////////////////////////////////////////////////////////////////////
+static FSharedBuffer* Writer_CreateSharedBuffer(uint32 SizeHint=0)
+{
+	uint32 BlockSize = GBlockSize;
+	if (SizeHint > GBlockSize)
+	{
+		BlockSize += SizeHint - GBlockSize;
+		BlockSize += sizeof(FSharedBuffer);
+		BlockSize += GBlockSize - 1;
+		BlockSize &= ~(GBlockSize - 1);
+	}
+
+	void* Block = Writer_MemoryAllocate(BlockSize, alignof(FSharedBuffer));
+	auto* Buffer = (FSharedBuffer*)(UPTRINT(Block) + BlockSize) - 1;
+
+	Buffer->Size = uint32(UPTRINT(Buffer) - UPTRINT(Block));
+	Buffer->Size -= sizeof(uint32); // to preceed event data with a small header when sending.
+	Buffer->Cursor = (Buffer->Size << FSharedBuffer::CursorShift) | FSharedBuffer::RefInit;
+	Buffer->Next = nullptr;
+	Buffer->Final = 0;
+
+	return Buffer;
+}
+
+////////////////////////////////////////////////////////////////////////////////
+FNextSharedBuffer Writer_NextSharedBuffer(FSharedBuffer* Buffer, int32 RegionStart, int32 NegSizeAndRef)
+{
+	FSharedBuffer* NextBuffer;
+	while (true)
+	{
+		bool bBufferOwner = (RegionStart >= 0);
+		if (LIKELY(bBufferOwner))
+		{
+			uint32 Size = -NegSizeAndRef >> FSharedBuffer::CursorShift;
+			NextBuffer = Writer_CreateSharedBuffer(Size);
+			Buffer->Next = NextBuffer;
+			Buffer->Final = RegionStart >> FSharedBuffer::CursorShift;
+			AtomicStoreRelease(&GSharedBuffer, NextBuffer);
+		}
+		else
+		{
+			for (;; PlatformYield())
+			{
+				NextBuffer = AtomicLoadAcquire(&GSharedBuffer);
+				if (NextBuffer != Buffer)
+				{
+					break;
+				}
+			}
+		}
+
+		AtomicAddRelease(&(Buffer->Cursor), int32(FSharedBuffer::RefBit));
+
+		// Try and allocate some space in the next buffer.
+		RegionStart = AtomicAddRelaxed(&(NextBuffer->Cursor), NegSizeAndRef);
+		if (LIKELY(RegionStart + NegSizeAndRef >= 0))
+		{
+			break;
+		}
+
+		Buffer = NextBuffer;
+	}
+
+	return { NextBuffer, RegionStart };
+}
+
+////////////////////////////////////////////////////////////////////////////////
+static void Writer_RetireSharedBufferImpl()
+{
+	// Send any unsent data.
+	uint8* Data = (uint8*)GTailBuffer - GTailBuffer->Size + GTailPreSent;
+	if (uint32 SendSize = UPTRINT(GTailBuffer) - UPTRINT(Data) - GTailBuffer->Final)
+	{
+		Writer_CacheData(Data, SendSize);
+	}
+
+	FSharedBuffer* Temp = GTailBuffer->Next;
+	void* Block = (uint8*)GTailBuffer - GTailBuffer->Size - sizeof(uint32);
+	Writer_MemoryFree(Block, GBlockSize);
+	GTailBuffer = Temp;
+	GTailPreSent = 0;
+}
+
+////////////////////////////////////////////////////////////////////////////////
+static void Writer_RetireSharedBuffer()
+{
+	// Spin until the buffer's no longer being used.
+	for (;; PlatformYield())
+	{
+		int32 TailCursor = AtomicLoadAcquire(&(GTailBuffer->Cursor));
+		if (LIKELY(((TailCursor + 1) & FSharedBuffer::RefInit) == 0))
+		{
+			break;
+		}
+	}
+
+	Writer_RetireSharedBufferImpl();
+}
+
+////////////////////////////////////////////////////////////////////////////////
+void Writer_UpdateSharedBuffers()
+{
+	FSharedBuffer* HeadBuffer = AtomicLoadAcquire(&GSharedBuffer);
+	while (true)
+	{
+		if (GTailBuffer != HeadBuffer)
+		{
+			Writer_RetireSharedBuffer();
+			continue;
+		}
+
+		int32 Cursor = AtomicLoadAcquire(&(HeadBuffer->Cursor));
+		if ((Cursor + 1) & FSharedBuffer::RefInit)
+		{
+			continue;
+		}
+
+		Cursor = Cursor >> FSharedBuffer::CursorShift;
+		if (Cursor < 0)
+		{
+			Writer_RetireSharedBufferImpl();
+			break;
+		}
+
+		uint32 PreSentBias = HeadBuffer->Size - GTailPreSent;
+		if (uint32 Sendable = PreSentBias - Cursor)
+		{
+			uint8* Data = (uint8*)(UPTRINT(HeadBuffer) - PreSentBias);
+			Writer_CacheData(Data, Sendable);
+			GTailPreSent += Sendable;
+		}
+
+		break;
+	}
+}
+
+////////////////////////////////////////////////////////////////////////////////
+void Writer_InitializeSharedBuffers()
+{
+	Writer_InitializeCache();
+
+	const uint32 BufferCount = 1;
+
+	FSharedBuffer* Buffers[BufferCount];
+	for (int i = 0; i < BufferCount; ++i)
+	{
+		Buffers[i] = Writer_CreateSharedBuffer();
+	}
+
+	for (int i = 1; i < BufferCount; ++i)
+	{
+		Buffers[i - 1]->Next = Buffers[i];
+	}
+
+	GTailBuffer = Buffers[0];
+	GTailPreSent = 0;
+
+	AtomicStoreRelease(&GSharedBuffer, Buffers[0]);
+}
+
+////////////////////////////////////////////////////////////////////////////////
+void Writer_ShutdownSharedBuffers()
+{
+	Writer_ShutdownCache();
+}
+
+} // namespace Private
+} // namespace Trace
+
+#endif // UE_TRACE_ENABLED
diff --git a/Engine/Source/Runtime/TraceLog/Private/Trace/TlsBuffer.cpp b/Engine/Source/Runtime/TraceLog/Private/Trace/TlsBuffer.cpp
--- a/Engine/Source/Runtime/TraceLog/Private/Trace/TlsBuffer.cpp
+++ b/Engine/Source/Runtime/TraceLog/Private/Trace/TlsBuffer.cpp
@@ -128,6 +128,27 @@ TRACELOG_API FWriteBuffer* Writer_NextBuffer(int32 Size)
 	return NextBuffer;
 }
 
+////////////////////////////////////////////////////////////////////////////////
+static bool Writer_DrainBuffer(uint32 ThreadId, FWriteBuffer* Buffer)
+{
+	uint8* Committed = AtomicLoadRelaxed((uint8**)&Buffer->Committed);
+
+	// Send as much as we can.
+	if (uint32 SizeToReap = uint32(Committed - Buffer->Reaped))
+	{
+#if TRACE_PRIVATE_PERF
+		BytesReaped += SizeToReap;
+		BytesSent += /*...*/
+#endif
+		Writer_SendData(ThreadId, Buffer->Reaped, SizeToReap);
+		Buffer->Reaped = Committed;
+	}
+
+	// Is this buffer still in use?
+	int32 EtxOffset = int32(AtomicLoadAcquire(&Buffer->EtxOffset));
+	return ((uint8*)Buffer - EtxOffset) > Committed;
+}
+
 ////////////////////////////////////////////////////////////////////////////////
 void Writer_DrainBuffers()
 {
@@ -193,22 +214,7 @@ void Writer_DrainBuffers()
 			// For each of the thread's buffers...
 			for (FWriteBuffer* __restrict NextBuffer; Buffer != nullptr; Buffer = NextBuffer)
 			{
-				uint8* Committed = AtomicLoadRelaxed((uint8**)&Buffer->Committed);
-
-				// Send as much as we can.
-				if (uint32 SizeToReap = uint32(Committed - Buffer->Reaped))
-				{
-#if TRACE_PRIVATE_PERF
-					BytesReaped += SizeToReap;
-					BytesSent += /*...*/
-#endif
-					Writer_SendData(ThreadId, Buffer->Reaped, SizeToReap);
-					Buffer->Reaped = Committed;
-				}
-
-				// Is this buffer still in use?
-				int32 EtxOffset = int32(AtomicLoadAcquire(&Buffer->EtxOffset));
-				if ((uint8*)Buffer - EtxOffset > Committed)
+				if (Writer_DrainBuffer(ThreadId, Buffer))
 				{
 					break;
 				}
diff --git a/Engine/Source/Runtime/TraceLog/Private/Trace/Trace.cpp b/Engine/Source/Runtime/TraceLog/Private/Trace/Trace.cpp
--- a/Engine/Source/Runtime/TraceLog/Private/Trace/Trace.cpp
+++ b/Engine/Source/Runtime/TraceLog/Private/Trace/Trace.cpp
@@ -4,8 +4,10 @@
 
 #if UE_TRACE_ENABLED
 
+#include "Misc/CString.h"
 #include "Trace/Detail/Channel.h"
 
+
 namespace Trace
 {
 
@@ -19,6 +21,10 @@ void	Writer_Update();
 bool	Writer_SendTo(const ANSICHAR*, uint32);
 bool	Writer_WriteTo(const ANSICHAR*);
 bool	Writer_IsTracing();
+bool	Writer_Stop();
+uint32	Writer_GetThreadId();
+
+extern FStatistics GTraceStatistics;
 
 } // namespace Private
 
@@ -60,6 +66,12 @@ void Update()
 	Private::Writer_Update();
 }
 
+////////////////////////////////////////////////////////////////////////////////
+void GetStatistics(FStatistics& Out)
+{
+	Out = Private::GTraceStatistics;
+}
+
 ////////////////////////////////////////////////////////////////////////////////
 bool SendTo(const TCHAR* InHost, uint32 Port)
 {
@@ -82,6 +94,12 @@ bool IsTracing()
 	return Private::Writer_IsTracing();
 }
 
+////////////////////////////////////////////////////////////////////////////////
+bool Stop()
+{
+	return Private::Writer_Stop();
+}
+
 ////////////////////////////////////////////////////////////////////////////////
 bool IsChannel(const TCHAR* ChannelName)
 {
@@ -103,33 +121,38 @@ bool ToggleChannel(const TCHAR* ChannelName, bool bEnabled)
 ////////////////////////////////////////////////////////////////////////////////
 UE_TRACE_CHANNEL_EXTERN(TraceLogChannel)
 
-UE_TRACE_EVENT_BEGIN($Trace, ThreadInfo)
+UE_TRACE_EVENT_BEGIN($Trace, ThreadInfo, NoSync|Important)
+	UE_TRACE_EVENT_FIELD(uint32, ThreadId)
 	UE_TRACE_EVENT_FIELD(uint32, SystemId)
 	UE_TRACE_EVENT_FIELD(int32, SortHint)
 	UE_TRACE_EVENT_FIELD(Trace::AnsiString, Name)
 UE_TRACE_EVENT_END()
 
-UE_TRACE_EVENT_BEGIN($Trace, ThreadGroupBegin)
+UE_TRACE_EVENT_BEGIN($Trace, ThreadGroupBegin, NoSync|Important)
 	UE_TRACE_EVENT_FIELD(Trace::AnsiString, Name)
 UE_TRACE_EVENT_END()
 
-UE_TRACE_EVENT_BEGIN($Trace, ThreadGroupEnd)
+UE_TRACE_EVENT_BEGIN($Trace, ThreadGroupEnd, NoSync|Important)
 UE_TRACE_EVENT_END()
 
 ////////////////////////////////////////////////////////////////////////////////
 void ThreadRegister(const TCHAR* Name, uint32 SystemId, int32 SortHint)
 {
-	UE_TRACE_LOG($Trace, ThreadInfo, TraceLogChannel)
+	uint32 ThreadId = Private::Writer_GetThreadId();
+	uint32 NameLen = FCString::Strlen(Name);
+	UE_TRACE_LOG($Trace, ThreadInfo, TraceLogChannel, NameLen * sizeof(ANSICHAR))
+		<< ThreadInfo.ThreadId(ThreadId)
 		<< ThreadInfo.SystemId(SystemId)
 		<< ThreadInfo.SortHint(SortHint)
-		<< ThreadInfo.Name(Name);
+		<< ThreadInfo.Name(Name, NameLen);
 }
 
 ////////////////////////////////////////////////////////////////////////////////
 void ThreadGroupBegin(const TCHAR* Name)
 {
-	UE_TRACE_LOG($Trace, ThreadGroupBegin, TraceLogChannel)
-		<< ThreadGroupBegin.Name(Name);
+	uint32 NameLen = FCString::Strlen(Name);
+	UE_TRACE_LOG($Trace, ThreadGroupBegin, TraceLogChannel, NameLen * sizeof(ANSICHAR))
+		<< ThreadGroupBegin.Name(Name, NameLen);
 }
 
 ////////////////////////////////////////////////////////////////////////////////
diff --git a/Engine/Source/Runtime/TraceLog/Private/Trace/Writer.cpp b/Engine/Source/Runtime/TraceLog/Private/Trace/Writer.cpp
--- a/Engine/Source/Runtime/TraceLog/Private/Trace/Writer.cpp
+++ b/Engine/Source/Runtime/TraceLog/Private/Trace/Writer.cpp
@@ -31,6 +31,10 @@ namespace Private {
 ////////////////////////////////////////////////////////////////////////////////
 int32			Encode(const void*, int32, void*, int32);
 uint32			Writer_SendData(uint32, uint8* __restrict, uint32);
+void			Writer_InitializeSharedBuffers();
+void			Writer_ShutdownSharedBuffers();
+void			Writer_UpdateSharedBuffers();
+void			Writer_CacheOnConnect();
 void			Writer_InitializePool();
 void			Writer_ShutdownPool();
 void			Writer_DrainBuffers();
@@ -43,21 +47,19 @@ void			Writer_ShutdownControl();
 
 ////////////////////////////////////////////////////////////////////////////////
 UE_TRACE_EVENT_BEGIN($Trace, NewTrace, NoSync)
+	UE_TRACE_EVENT_FIELD(uint64, StartCycle)
+	UE_TRACE_EVENT_FIELD(uint64, CycleFrequency)
 	UE_TRACE_EVENT_FIELD(uint32, Serial)
 	UE_TRACE_EVENT_FIELD(uint16, UserUidBias)
 	UE_TRACE_EVENT_FIELD(uint16, Endian)
 	UE_TRACE_EVENT_FIELD(uint8, PointerSize)
 UE_TRACE_EVENT_END()
 
-UE_TRACE_EVENT_BEGIN($Trace, Timing, NoSync)
-	UE_TRACE_EVENT_FIELD(uint64, StartCycle)
-	UE_TRACE_EVENT_FIELD(uint64, CycleFrequency)
-UE_TRACE_EVENT_END()
-
 
 
 ////////////////////////////////////////////////////////////////////////////////
 static bool						GInitialized;		// = false;
+FStatistics						GTraceStatistics;	// = {};
 uint64							GStartCycle;		// = 0;
 TRACELOG_API uint32 volatile	GLogSerial;			// = 0;
 
@@ -147,6 +149,10 @@ void* Writer_MemoryAllocate(SIZE_T Size, uint32 Alignment)
 		Writer_SendData(ThreadId, TraceData.GetData(), TraceData.GetSize());
 	}
 
+#if !TRACE_PRIVATE_STATISTICS
+	AtomicIncrement(&GStatistics.MemoryUsed, Size);
+#endif
+
 	return Ret;
 }
 
@@ -181,6 +187,10 @@ void Writer_MemoryFree(void* Address, uint32 Size)
 		Writer_SendData(ThreadId, TraceData.GetData(), TraceData.GetSize());
 	}
 #endif // TRACE_PRIVATE_STOMP
+
+#if !TRACE_PRIVATE_STATISTICS
+	AtomicIncrement(&GStatistics.MemoryUsed, -int64(Size));
+#endif
 }
 
 
@@ -192,6 +202,10 @@ UPTRINT							GPendingDataHandle;	// = 0
 ////////////////////////////////////////////////////////////////////////////////
 void Writer_SendDataRaw(const void* Data, uint32 Size)
 {
+#if TRACE_PRIVATE_STATISTICS
+	GTraceStatistics.BytesSent += Size;
+#endif
+
 	if (!IoWrite(GDataHandle, Data, Size))
 	{
 		IoClose(GDataHandle);
@@ -202,6 +216,10 @@ void Writer_SendDataRaw(const void* Data, uint32 Size)
 ////////////////////////////////////////////////////////////////////////////////
 uint32 Writer_SendData(uint32 ThreadId, uint8* __restrict Data, uint32 Size)
 {
+#if TRACE_PRIVATE_STATISTICS
+	GTraceStatistics.BytesTraced += Size;
+#endif
+
 	if (!GDataHandle)
 	{
 		return 0;
@@ -303,8 +321,8 @@ static void Writer_DescribeAnnounce()
 		return;
 	}
 
-	Writer_DescribeEvents();
 	Writer_AnnounceChannels();
+	Writer_DescribeEvents();
 }
 
 
@@ -316,15 +334,9 @@ static void Writer_LogHeader()
 		<< NewTrace.Serial(uint32(GLogSerial)) // should really atomic-load-relaxed here...
 		<< NewTrace.UserUidBias(EKnownEventUids::User)
 		<< NewTrace.Endian(0x524d)
-		<< NewTrace.PointerSize(sizeof(void*));
-}
-
-////////////////////////////////////////////////////////////////////////////////
-static void Writer_LogTimingHeader()
-{
-	UE_TRACE_LOG($Trace, Timing, TraceLogChannel)
-		<< Timing.StartCycle(GStartCycle)
-		<< Timing.CycleFrequency(TimeGetFrequency());
+		<< NewTrace.PointerSize(sizeof(void*))
+		<< NewTrace.StartCycle(GStartCycle)
+		<< NewTrace.CycleFrequency(TimeGetFrequency());
 }
 
 ////////////////////////////////////////////////////////////////////////////////
@@ -335,6 +347,18 @@ static bool Writer_UpdateConnection()
 		return false;
 	}
 
+	// Is this a close request?
+	if (GPendingDataHandle == ~0ull)
+	{
+		if (GDataHandle)
+		{
+			IoClose(GDataHandle);
+		}
+
+		GDataHandle = 0;
+		GPendingDataHandle = 0;
+	}
+
 	// Reject the pending connection if we've already got a connection
 	if (GDataHandle)
 	{
@@ -364,16 +388,22 @@ static bool Writer_UpdateConnection()
 		return false;
 	}
 
+	// Reset statistics.
+	GTraceStatistics.BytesSent = 0;
+	GTraceStatistics.BytesTraced = 0;
+
 	// Send the header events
 	TWriteBufferRedirect<512> HeaderEvents;
 	Writer_LogHeader();
-	Writer_LogTimingHeader();
 	HeaderEvents.Close();
 
+	FEventNode::OnConnect();
 	Writer_DescribeEvents();
 
 	Writer_SendData(HeaderEvents.GetData(), HeaderEvents.GetSize());
 
+	Writer_CacheOnConnect();
+
 	return true;
 }
 
@@ -389,6 +419,7 @@ static void Writer_WorkerUpdate()
 	Writer_UpdateControl();
 	Writer_UpdateConnection();
 	Writer_DescribeAnnounce();
+	Writer_UpdateSharedBuffers();
 	Writer_DrainBuffers();
 }
 
@@ -448,6 +479,7 @@ static void Writer_InternalInitializeImpl()
 	GInitialized = true;
 	GStartCycle = TimeGetTimestamp();
 
+	Writer_InitializeSharedBuffers();
 	Writer_InitializePool();
 	Writer_InitializeControl();
 }
@@ -479,6 +511,7 @@ static void Writer_InternalShutdown()
 
 	Writer_ShutdownControl();
 	Writer_ShutdownPool();
+	Writer_ShutdownSharedBuffers();
 
 	GInitialized = false;
 }
@@ -588,6 +621,18 @@ bool Writer_IsTracing()
 	return (GDataHandle != 0);
 }
 
+////////////////////////////////////////////////////////////////////////////////
+bool Writer_Stop()
+{
+	if (GPendingDataHandle || !GDataHandle)
+	{
+		return false;
+	}
+
+	GPendingDataHandle = ~UPTRINT(0);
+	return true;
+}
+
 } // namespace Private
 } // namespace Trace
 
diff --git a/Engine/Source/Runtime/TraceLog/Public/Trace/Detail/EventNode.h b/Engine/Source/Runtime/TraceLog/Public/Trace/Detail/EventNode.h
--- a/Engine/Source/Runtime/TraceLog/Public/Trace/Detail/EventNode.h
+++ b/Engine/Source/Runtime/TraceLog/Public/Trace/Detail/EventNode.h
@@ -40,6 +40,7 @@ public:
 	};
 
 	static FIter			ReadNew();
+	static void				OnConnect();
 	TRACELOG_API uint32		Initialize(const FEventInfo* InInfo);
 	void					Describe() const;
 	uint32					GetUid() const { return Uid; }
diff --git a/Engine/Source/Runtime/TraceLog/Public/Trace/Detail/Field.h b/Engine/Source/Runtime/TraceLog/Public/Trace/Detail/Field.h
--- a/Engine/Source/Runtime/TraceLog/Public/Trace/Detail/Field.h
+++ b/Engine/Source/Runtime/TraceLog/Public/Trace/Detail/Field.h
@@ -11,6 +11,11 @@
 #include "Templates/UnrealTemplate.h"
 #include "Writer.inl"
 
+/* Statically sized fields (e.g. UE_TRACE_EVENT_FIELD(float[4], Colours)) are
+ * not supported as yet. No call for them. The following define is used to track
+ * where and partially how to implement them */
+#define STATICALLY_SIZED_ARRAY_FIELDS_SUPPORT 0
+
 namespace Trace {
 namespace Private {
 
@@ -48,7 +53,7 @@ struct TFieldType<T[]>
 	};
 };
 
-#if 0
+#if STATICALLY_SIZED_ARRAY_FIELDS_SUPPORT
 template <typename T, int N>
 struct TFieldType<T[N]>
 {
@@ -58,7 +63,7 @@ struct TFieldType<T[N]>
 		Size = sizeof(T[N]),
 	};
 };
-#endif // 0
+#endif // STATICALLY_SIZED_ARRAY_FIELDS_SUPPORT
 
 template <> struct TFieldType<Trace::AnsiString> { enum { Tid  = int(EFieldType::AnsiString), Size = 0, }; };
 template <> struct TFieldType<Trace::WideString> { enum { Tid  = int(EFieldType::WideString), Size = 0, }; };
@@ -106,8 +111,9 @@ template <int InIndex, int InOffset, typename Type> struct TField;
 
 enum class EIndexPack
 {
-	FieldCountMask	= 0xff,
-	MaybeHasAux		= 0x100,
+	NumFieldsShift	= 8,
+	NumFieldsMask	= (1 << NumFieldsShift) - 1,
+	AuxFieldCounter	= 1 << NumFieldsShift,
 };
 
 ////////////////////////////////////////////////////////////////////////////////
@@ -119,7 +125,7 @@ enum class EIndexPack
 			Tid		= TFieldType<Type>::Tid, \
 			Size	= TFieldType<Type>::Size, \
 		}; \
-		static_assert((Index & int(EIndexPack::FieldCountMask)) <= 127, "Trace events may only have up to a maximum of 127 fields"); \
+		static_assert((Index & int(EIndexPack::NumFieldsMask)) <= 127, "Trace events may only have up to a maximum of 127 fields"); \
 	private: \
 		FFieldDesc FieldDesc; \
 	public: \
@@ -132,81 +138,30 @@ enum class EIndexPack
 template <int InIndex, int InOffset, typename Type>
 struct TField<InIndex, InOffset, Type[]>
 {
-	TRACE_PRIVATE_FIELD(InIndex|int(EIndexPack::MaybeHasAux), InOffset, Type[]);
-
-	static void Set(uint8*, Type const* Data, int32 Count)
-	{
-		if (Count > 0)
-		{
-			int32 Size = (Count * sizeof(Type)) & (FAuxHeader::SizeLimit - 1) & ~(sizeof(Type) - 1);
-			Private::Field_WriteAuxData(Index, (const uint8*)Data, Size);
-		}
-	}
+	TRACE_PRIVATE_FIELD(InIndex + int(EIndexPack::AuxFieldCounter), InOffset, Type[]);
 };
 
-#if 0
+#if STATICALLY_SIZED_ARRAY_FIELDS_SUPPORT
 ////////////////////////////////////////////////////////////////////////////////
 template <int InIndex, int InOffset, typename Type, int Count>
 struct TField<InIndex, InOffset, Type[Count]>
 {
 	TRACE_PRIVATE_FIELD(InIndex, InOffset, Type[Count]);
-
-	static void Set(uint8*, Type const* Data, int Count)
-	{
-	}
 };
-#endif
+#endif // STATICALLY_SIZED_ARRAY_FIELDS_SUPPORT
 
 ////////////////////////////////////////////////////////////////////////////////
 template <int InIndex, int InOffset>
 struct TField<InIndex, InOffset, Trace::AnsiString>
 {
-	TRACE_PRIVATE_FIELD(InIndex|int(EIndexPack::MaybeHasAux), InOffset, Trace::AnsiString);
-
-	FORCENOINLINE static void Set(uint8*, const TCHAR* String, int32 Length=-1)
-	{
-		if (Length < 0)
-		{
-			Length = 0;
-			for (const TCHAR* c = String; *c; ++c, ++Length);
-		}
-
-		Private::Field_WriteStringAnsi(Index, String, Length);
-	}
-
-	FORCENOINLINE static void Set(uint8*, const ANSICHAR* String, int32 Length=-1)
-	{
-		if (Length < 0)
-		{
-			Length = int32(strlen(String));
-		}
-
-		if (Length)
-		{
-			Private::Field_WriteStringAnsi(Index, String, Length);
-		}
-	}
+	TRACE_PRIVATE_FIELD(InIndex + int(EIndexPack::AuxFieldCounter), InOffset, Trace::AnsiString);
 };
 
 ////////////////////////////////////////////////////////////////////////////////
 template <int InIndex, int InOffset>
 struct TField<InIndex, InOffset, Trace::WideString>
 {
-	TRACE_PRIVATE_FIELD(InIndex|int(EIndexPack::MaybeHasAux), InOffset, Trace::WideString);
-
-	FORCENOINLINE static void Set(uint8*, const TCHAR* String, int32 Length=-1)
-	{
-		if (Length < 0)
-		{
-			Length = 0;
-			for (const TCHAR* c = String; *c; ++c, ++Length);
-		}
-
-		if (Length)
-		{
-			Private::Field_WriteStringWide(Index, String, Length);
-		}
-	}
+	TRACE_PRIVATE_FIELD(InIndex + int(EIndexPack::AuxFieldCounter), InOffset, Trace::WideString);
 };
 
 ////////////////////////////////////////////////////////////////////////////////
@@ -214,11 +169,6 @@ template <int InIndex, int InOffset, typename Type>
 struct TField
 {
 	TRACE_PRIVATE_FIELD(InIndex, InOffset, Type);
-
-	static void Set(uint8* Dest, const Type& __restrict Value)
-	{
-		::memcpy(Dest + Offset, &Value, Size);
-	}
 };
 
 #undef TRACE_PRIVATE_FIELD
@@ -228,14 +178,14 @@ struct TField
 ////////////////////////////////////////////////////////////////////////////////
 // Used to terminate the field list and determine an event's size.
 enum EventProps {};
-template <int InFieldCount, int InSize>
-struct TField<InFieldCount, InSize, EventProps>
+template <int InNumFields, int InSize>
+struct TField<InNumFields, InSize, EventProps>
 {
 	enum : uint16
 	{
-		FieldCount	= (InFieldCount & int(EIndexPack::FieldCountMask)),
-		Size		= InSize,
-		MaybeHasAux	= !!(InFieldCount & int(EIndexPack::MaybeHasAux)),
+		NumFields		= InNumFields & int(EIndexPack::NumFieldsMask),
+		Size			= InSize,
+		NumAuxFields	= (InNumFields >> int(EIndexPack::NumFieldsShift)) & int(EIndexPack::NumFieldsMask),
 	};
 };
 
@@ -245,16 +195,7 @@ enum Attachment {};
 template <int InOffset>
 struct TField<0, InOffset, Attachment>
 {
-	template <typename LambdaType>
-	static void Set(uint8* Dest, LambdaType&& Lambda)
-	{
-		Lambda(Dest + InOffset);
-	}
-
-	static void Set(uint8* Dest, const void* Data, uint32 Size)
-	{
-		::memcpy(Dest + InOffset, Data, Size);
-	}
+	enum { Offset = InOffset };
 };
 
 } // namespace Trace
diff --git a/Engine/Source/Runtime/TraceLog/Public/Trace/Detail/Important/ImportantLogScope.h b/Engine/Source/Runtime/TraceLog/Public/Trace/Detail/Important/ImportantLogScope.h
new file mode 100644
--- /dev/null
+++ b/Engine/Source/Runtime/TraceLog/Public/Trace/Detail/Important/ImportantLogScope.h
@@ -0,0 +1,37 @@
+// Copyright Epic Games, Inc. All Rights Reserved.
+
+#pragma once
+
+#include "Trace/Config.h"
+
+#if UE_TRACE_ENABLED
+
+namespace Trace {
+namespace Private {
+
+////////////////////////////////////////////////////////////////////////////////
+class FImportantLogScope
+{
+public:
+	template <typename EventType>
+	static FImportantLogScope	Enter();
+	template <typename EventType>
+	static FImportantLogScope	Enter(uint32 ArrayDataSize);
+	void						operator += (const FImportantLogScope&) const;
+	const FImportantLogScope&	operator << (bool) const	{ return *this; }
+	constexpr explicit			operator bool () const		{ return true; }
+
+	template <typename FieldMeta, typename Type>
+	struct FFieldSet;
+
+private:
+	static FImportantLogScope	EnterImpl(uint32 Uid, uint32 Size);
+	uint8*						Ptr;
+	int32						BufferOffset;
+	int32						AuxCursor;
+};
+
+} // namespace Private
+} // namespace Trace
+
+#endif // UE_TRACE_ENABLED
diff --git a/Engine/Source/Runtime/TraceLog/Public/Trace/Detail/Important/ImportantLogScope.inl b/Engine/Source/Runtime/TraceLog/Public/Trace/Detail/Important/ImportantLogScope.inl
new file mode 100644
--- /dev/null
+++ b/Engine/Source/Runtime/TraceLog/Public/Trace/Detail/Important/ImportantLogScope.inl
@@ -0,0 +1,208 @@
+// Copyright Epic Games, Inc. All Rights Reserved.
+
+#pragma once
+
+#include "Trace/Config.h"
+
+#if UE_TRACE_ENABLED
+
+#include "ImportantLogScope.h"
+#include "SharedBuffer.h"
+#include "Trace/Detail/Protocol.h"
+#include "Trace/Detail/Writer.inl"
+#include "Trace/Detail/Field.h"
+#include "Trace/Detail/EventNode.h"
+
+namespace Trace {
+namespace Private {
+
+////////////////////////////////////////////////////////////////////////////////
+extern TRACELOG_API FSharedBuffer* volatile GSharedBuffer;
+TRACELOG_API FNextSharedBuffer				Writer_NextSharedBuffer(FSharedBuffer*, int32, int32);
+
+
+
+////////////////////////////////////////////////////////////////////////////////
+template <class T>
+FORCENOINLINE FImportantLogScope FImportantLogScope::Enter(uint32 ArrayDataSize)
+{
+	static_assert(T::EventFlags & FEventInfo::Flag_MaybeHasAux, "Only important trace events with array-type fields need a size parameter to UE_TRACE_LOG()");
+
+	ArrayDataSize += sizeof(FAuxHeader) * T::EventProps_Meta::NumAuxFields;
+	ArrayDataSize += 1; // null terminator
+
+	uint32 Size = T::GetSize();
+	uint32 Uid = T::GetUid();
+	FImportantLogScope Ret = EnterImpl(Uid, Size + ArrayDataSize);
+
+	Ret.AuxCursor += Size;
+	Ret.Ptr[Ret.AuxCursor] = 0; // null terminator
+	return Ret;
+}
+
+////////////////////////////////////////////////////////////////////////////////
+template <class T>
+inline FImportantLogScope FImportantLogScope::Enter()
+{
+	static_assert(!(T::EventFlags & FEventInfo::Flag_MaybeHasAux), "Important trace events with array-type fields must be traced with UE_TRACE_LOG(Logger, Event, Channel, ArrayDataSize)");
+
+	uint32 Size = T::GetSize();
+	uint32 Uid = T::GetUid();
+	return EnterImpl(Uid, Size);
+}
+
+////////////////////////////////////////////////////////////////////////////////
+inline FImportantLogScope FImportantLogScope::EnterImpl(uint32 Uid, uint32 Size)
+{
+	FSharedBuffer* Buffer = AtomicLoadAcquire(&GSharedBuffer);
+
+	int32 AllocSize = Size;
+	AllocSize += sizeof(FEventHeader);
+
+	// Claim some space in the buffer
+	int32 NegSizeAndRef = 0 - ((AllocSize << FSharedBuffer::CursorShift) | FSharedBuffer::RefBit);
+	int32 RegionStart = AtomicAddRelaxed(&(Buffer->Cursor), NegSizeAndRef);
+
+	if (UNLIKELY(RegionStart + NegSizeAndRef < 0))
+	{
+		FNextSharedBuffer Next = Writer_NextSharedBuffer(Buffer, RegionStart, NegSizeAndRef);
+		Buffer = Next.Buffer;
+		RegionStart = Next.RegionStart;
+	}
+
+	int32 Bias = (RegionStart >> FSharedBuffer::CursorShift);
+	uint8* Out = (uint8*)Buffer - Bias;
+
+	// Event header
+	auto* Header = (uint16*)Out;
+	Header[0] = uint16(Uid | uint32(EKnownEventUids::Flag_TwoByteUid));
+	Header[1] = uint16(Size);
+
+	FImportantLogScope Ret;
+	Ret.Ptr = (uint8*)(Header + 2);
+	Ret.BufferOffset = int32(PTRINT(Buffer) - PTRINT(Ret.Ptr));
+	Ret.AuxCursor = 0;
+	return Ret;
+}
+
+////////////////////////////////////////////////////////////////////////////////
+inline void FImportantLogScope::operator += (const FImportantLogScope&) const
+{
+	auto* Buffer = (FSharedBuffer*)(Ptr + BufferOffset);
+	AtomicAddRelease(&(Buffer->Cursor), int32(FSharedBuffer::RefBit));
+}
+
+
+
+////////////////////////////////////////////////////////////////////////////////
+template <typename FieldMeta, typename Type>
+struct FImportantLogScope::FFieldSet
+{
+	static void Impl(FImportantLogScope* Scope, const Type& Value)
+	{
+		uint8* Dest = (uint8*)(Scope->Ptr) + FieldMeta::Offset;
+		::memcpy(Dest, &Value, sizeof(Type));
+	}
+};
+
+////////////////////////////////////////////////////////////////////////////////
+template <typename FieldMeta, typename Type>
+struct FImportantLogScope::FFieldSet<FieldMeta, Type[]>
+{
+	static void Impl(FImportantLogScope* Scope, Type const* Data, int32 Num)
+	{
+		uint32 Size = Num * sizeof(Type);
+
+		auto* Header = (FAuxHeader*)(Scope->Ptr + Scope->AuxCursor);
+		Header->Size = Size << 8;
+		Header->FieldIndex = uint8(0x80 | (FieldMeta::Index & int32(EIndexPack::NumFieldsMask)));
+
+		memcpy(Header + 1, Data, Size);
+
+		Scope->AuxCursor += sizeof(FAuxHeader) + Size;
+		Scope->Ptr[Scope->AuxCursor] = 0; // null terminator
+	}
+};
+
+////////////////////////////////////////////////////////////////////////////////
+template <typename FieldMeta>
+struct FImportantLogScope::FFieldSet<FieldMeta, AnsiString>
+{
+	static void Impl(FImportantLogScope* Scope, const ANSICHAR* String, int32 Length=-1)
+	{
+		if (Length < 0)
+		{
+			Length = int32(strlen(String));
+		}
+
+		auto* Header = (FAuxHeader*)(Scope->Ptr + Scope->AuxCursor);
+		Header->Size = Length << 8;
+		Header->FieldIndex = uint8(0x80 | (FieldMeta::Index & int32(EIndexPack::NumFieldsMask)));
+
+		memcpy(Header + 1, String, Length);
+
+		Scope->AuxCursor += sizeof(FAuxHeader) + Length;
+		Scope->Ptr[Scope->AuxCursor] = 0; // null terminator
+	}
+
+	static void Impl(FImportantLogScope* Scope, const TCHAR* String, int32 Length=-1)
+	{
+		if (Length < 0)
+		{
+			Length = 0;
+			for (const TCHAR* c = String; *c; ++c, ++Length);
+		}
+
+		auto* Header = (FAuxHeader*)(Scope->Ptr + Scope->AuxCursor);
+		Header->Size = Length << 8;
+		Header->FieldIndex = uint8(0x80 | (FieldMeta::Index & int32(EIndexPack::NumFieldsMask)));
+
+		auto* Out = (int8*)(Header + 1);
+		for (int32 i = 0; i < Length; ++i)
+		{
+			*Out = int8(*String);
+			++Out;
+			++String;
+		}
+
+		Scope->AuxCursor += sizeof(FAuxHeader) + Length;
+		Scope->Ptr[Scope->AuxCursor] = 0; // null terminator
+	}
+};
+
+////////////////////////////////////////////////////////////////////////////////
+template <typename FieldMeta>
+struct FImportantLogScope::FFieldSet<FieldMeta, WideString>
+{
+	static void Impl(FImportantLogScope* Scope, const TCHAR* String, int32 Length=-1)
+	{
+		if (Length < 0)
+		{
+			Length = 0;
+			for (const TCHAR* c = String; *c; ++c, ++Length);
+		}
+
+		uint32 Size = Length * sizeof(TCHAR);
+
+		auto* Header = (FAuxHeader*)(Scope->Ptr + Scope->AuxCursor);
+		Header->Size = Size << 8;
+		Header->FieldIndex = uint8(0x80 | (FieldMeta::Index & int32(EIndexPack::NumFieldsMask)));
+
+		memcpy(Header + 1, String, Size);
+
+		Scope->AuxCursor += sizeof(FAuxHeader) + Size;
+		Scope->Ptr[Scope->AuxCursor] = 0; // null terminator
+	}
+};
+
+////////////////////////////////////////////////////////////////////////////////
+template <typename FieldMeta>
+struct FImportantLogScope::FFieldSet<FieldMeta, Attachment>
+{
+	static void Impl(...); /* not implemented */
+};
+
+} // namespace Private
+} // namespace Trace
+
+#endif // UE_TRACE_ENABLED
diff --git a/Engine/Source/Runtime/TraceLog/Public/Trace/Detail/Important/SharedBuffer.h b/Engine/Source/Runtime/TraceLog/Public/Trace/Detail/Important/SharedBuffer.h
new file mode 100644
--- /dev/null
+++ b/Engine/Source/Runtime/TraceLog/Public/Trace/Detail/Important/SharedBuffer.h
@@ -0,0 +1,37 @@
+// Copyright Epic Games, Inc. All Rights Reserved.
+
+#pragma once
+
+#include "Trace/Config.h"
+
+#if UE_TRACE_ENABLED
+
+namespace Trace {
+namespace Private {
+
+////////////////////////////////////////////////////////////////////////////////
+struct FSharedBuffer
+{
+	enum : uint32 {	CursorShift	= 10 };
+	enum : uint32 {	RefBit		= 1 << 0 };
+	enum : uint32 {	RefInit		= (1 << CursorShift) - 1 };
+	enum : uint32 {	MaxSize		= 1 << (32 - CursorShift - 1) };
+
+	int32 volatile	Cursor; // also packs in a ref count.
+	uint32			Size;
+	uint32			Final;
+	uint32			_Unused;
+	FSharedBuffer*	Next;
+};
+
+////////////////////////////////////////////////////////////////////////////////
+struct FNextSharedBuffer
+{
+	FSharedBuffer*	Buffer;
+	int32			RegionStart;
+};
+
+} // namespace Private
+} // namespace Trace
+
+#endif // UE_TRACE_ENABLED
diff --git a/Engine/Source/Runtime/TraceLog/Public/Trace/Detail/LogScope.h b/Engine/Source/Runtime/TraceLog/Public/Trace/Detail/LogScope.h
--- a/Engine/Source/Runtime/TraceLog/Public/Trace/Detail/LogScope.h
+++ b/Engine/Source/Runtime/TraceLog/Public/Trace/Detail/LogScope.h
@@ -17,30 +17,32 @@ struct FWriteBuffer;
 class FLogScope
 {
 public:
-	template <uint32 Flags>
-	static FLogScope		Enter(uint32 Uid, uint32 Size);
-	uint8*					GetPointer() const;
+	template <typename EventType>
+	static FLogScope		Enter(uint32 ExtraSize=0);
+	template <typename EventType>
+	static FLogScope		ScopedEnter(uint32 ExtraSize=0);
+	template <typename EventType>
+	static FLogScope		ScopedStampedEnter(uint32 ExtraSize=0);
+	void*					GetPointer() const			{ return Ptr; }
 	void					Commit() const;
 	void					operator += (const FLogScope&) const;
 	const FLogScope&		operator << (bool) const	{ return *this; }
 	constexpr explicit		operator bool () const		{ return true; }
 
+	template <typename FieldMeta, typename Type>
+	struct FFieldSet;
+
 private:
+	template <uint32 Flags>
+	static FLogScope		EnterImpl(uint32 Uid, uint32 Size);
 	template <class T> void	EnterPrelude(uint32 Size, bool bMaybeHasAux);
 	void					Enter(uint32 Uid, uint32 Size, bool bMaybeHasAux);
 	void					EnterNoSync(uint32 Uid, uint32 Size, bool bMaybeHasAux);
-	struct
-	{
-		uint8*				Ptr;
-		FWriteBuffer*		Buffer;
-	}						Instance;
+	uint8*					Ptr;
+	FWriteBuffer*			Buffer;
 };
 
-////////////////////////////////////////////////////////////////////////////////
-class FImportantLogScope
-	: public FLogScope
-{
-};
+
 
 ////////////////////////////////////////////////////////////////////////////////
 class FScopedLogScope
@@ -60,17 +62,6 @@ public:
 	bool	bActive = false;
 };
 
-
-
-////////////////////////////////////////////////////////////////////////////////
-template <class T>
-struct TLogScope
-{
-	static auto			Enter(uint32 ExtraSize=0);
-	static FLogScope	ScopedEnter(uint32 ExtraSize=0);
-	static FLogScope	ScopedStampedEnter(uint32 ExtraSize=0);
-};
-
 } // namespace Private
 } // namespace Trace
 
diff --git a/Engine/Source/Runtime/TraceLog/Public/Trace/Detail/LogScope.inl b/Engine/Source/Runtime/TraceLog/Public/Trace/Detail/LogScope.inl
--- a/Engine/Source/Runtime/TraceLog/Public/Trace/Detail/LogScope.inl
+++ b/Engine/Source/Runtime/TraceLog/Public/Trace/Detail/LogScope.inl
@@ -18,16 +18,9 @@ namespace Private {
 ////////////////////////////////////////////////////////////////////////////////
 extern TRACELOG_API uint32 volatile	GLogSerial;
 
-////////////////////////////////////////////////////////////////////////////////
-inline uint8* FLogScope::GetPointer() const
-{
-	return Instance.Ptr;
-}
-
 ////////////////////////////////////////////////////////////////////////////////
 inline void FLogScope::Commit() const
 {
-	FWriteBuffer* Buffer = Instance.Buffer;
 	AtomicStoreRelease((uint8**) &(Buffer->Committed), Buffer->Cursor);
 }
 
@@ -39,7 +32,7 @@ inline void FLogScope::operator += (const FLogScope&) const
 
 ////////////////////////////////////////////////////////////////////////////////
 template <uint32 Flags>
-inline FLogScope FLogScope::Enter(uint32 Uid, uint32 Size)
+inline FLogScope FLogScope::EnterImpl(uint32 Uid, uint32 Size)
 {
 	FLogScope Ret;
 	bool bMaybeHasAux = (Flags & FEventInfo::Flag_MaybeHasAux) != 0;
@@ -58,9 +51,9 @@ inline FLogScope FLogScope::Enter(uint32 Uid, uint32 Size)
 template <class HeaderType>
 inline void FLogScope::EnterPrelude(uint32 Size, bool bMaybeHasAux)
 {
-	uint32 AllocSize = sizeof(HeaderType) + Size + int(bMaybeHasAux);
+	uint32 AllocSize = sizeof(HeaderType) + Size + int32(bMaybeHasAux);
 
-	FWriteBuffer* Buffer = Writer_GetBuffer();
+	Buffer = Writer_GetBuffer();
 	Buffer->Cursor += AllocSize;
 	if (UNLIKELY(Buffer->Cursor > (uint8*)Buffer))
 	{
@@ -73,8 +66,7 @@ inline void FLogScope::EnterPrelude(uint32 Size, bool bMaybeHasAux)
 		Buffer->Cursor[-1] = 0;
 	}
 
-	uint8* Cursor = Buffer->Cursor - Size - int(bMaybeHasAux);
-	Instance = {Cursor, Buffer};
+	Ptr = Buffer->Cursor - Size - int32(bMaybeHasAux);
 }
 
 ////////////////////////////////////////////////////////////////////////////////
@@ -83,10 +75,10 @@ inline void FLogScope::Enter(uint32 Uid, uint32 Size, bool bMaybeHasAux)
 	EnterPrelude<FEventHeaderSync>(Size, bMaybeHasAux);
 
 	// Event header
-	auto* Header = (uint16*)(Instance.Ptr - sizeof(FEventHeaderSync::SerialHigh)); // FEventHeader1
+	auto* Header = (uint16*)(Ptr - sizeof(FEventHeaderSync::SerialHigh));
 	*(uint32*)(Header - 1) = uint32(AtomicAddRelaxed(&GLogSerial, 1u));
 	Header[-2] = uint16(Size);
-	Header[-3] = uint16(Uid)|int(EKnownEventUids::Flag_TwoByteUid);
+	Header[-3] = uint16(Uid)|int32(EKnownEventUids::Flag_TwoByteUid);
 }
 
 ////////////////////////////////////////////////////////////////////////////////
@@ -95,9 +87,9 @@ inline void FLogScope::EnterNoSync(uint32 Uid, uint32 Size, bool bMaybeHasAux)
 	EnterPrelude<FEventHeader>(Size, bMaybeHasAux);
 
 	// Event header
-	auto* Header = (uint16*)(Instance.Ptr);
+	auto* Header = (uint16*)(Ptr);
 	Header[-1] = uint16(Size);
-	Header[-2] = uint16(Uid)|int(EKnownEventUids::Flag_TwoByteUid);
+	Header[-2] = uint16(Uid)|int32(EKnownEventUids::Flag_TwoByteUid);
 }
 
 
@@ -165,28 +157,19 @@ inline void FScopedStampedLogScope::SetActive()
 
 
 
-////////////////////////////////////////////////////////////////////////////////
-template <bool>	struct TLogScopeSelector;
-template <>		struct TLogScopeSelector<false>	{ typedef FLogScope Type; };
-template <>		struct TLogScopeSelector<true>	{ typedef FImportantLogScope Type; };
-
 ////////////////////////////////////////////////////////////////////////////////
 template <class T>
-auto TLogScope<T>::Enter(uint32 ExtraSize)
+FORCENOINLINE FLogScope FLogScope::Enter(uint32 ExtraSize)
 {
 	uint32 Size = T::GetSize() + ExtraSize;
 	uint32 Uid = T::GetUid();
-
-	using LogScopeType = typename TLogScopeSelector<T::bIsImportant>::Type;
-	return LogScopeType::template Enter<T::EventFlags>(Uid, Size);
+	return EnterImpl<T::EventFlags>(Uid, Size);
 }
 
 ////////////////////////////////////////////////////////////////////////////////
 template <class T>
-FLogScope TLogScope<T>::ScopedEnter(uint32 ExtraSize)
+FORCENOINLINE FLogScope FLogScope::ScopedEnter(uint32 ExtraSize)
 {
-	static_assert(!T::bIsImportant, "Important events cannot be logged with scope");
-
 	uint8 EnterUid = uint8(EKnownEventUids::EnterScope << EKnownEventUids::_UidShift);
 
 	FWriteBuffer* Buffer = Writer_GetBuffer();
@@ -200,15 +183,13 @@ FLogScope TLogScope<T>::ScopedEnter(uint32 ExtraSize)
 
 	AtomicStoreRelease((uint8**) &(Buffer->Committed), Buffer->Cursor);
 
-	return Enter(ExtraSize);
+	return Enter<T>(ExtraSize);
 }
 
 ////////////////////////////////////////////////////////////////////////////////
 template <class T>
-FLogScope TLogScope<T>::ScopedStampedEnter(uint32 ExtraSize)
+FORCENOINLINE FLogScope FLogScope::ScopedStampedEnter(uint32 ExtraSize)
 {
-	static_assert(!T::bIsImportant, "Important events cannot be logged with scope");
-
 	uint64 Stamp;
 
 	FWriteBuffer* Buffer = Writer_GetBuffer();
@@ -225,9 +206,105 @@ FLogScope TLogScope<T>::ScopedStampedEnter(uint32 ExtraSize)
 
 	AtomicStoreRelease((uint8**) &(Buffer->Committed), Buffer->Cursor);
 
-	return Enter(ExtraSize);
+	return Enter<T>(ExtraSize);
 }
 
+
+
+////////////////////////////////////////////////////////////////////////////////
+template <typename FieldMeta, typename Type>
+struct FLogScope::FFieldSet
+{
+	static void Impl(FLogScope* Scope, const Type& Value)
+	{
+		uint8* Dest = (uint8*)(Scope->Ptr) + FieldMeta::Offset;
+		::memcpy(Dest, &Value, sizeof(Type));
+	}
+};
+
+////////////////////////////////////////////////////////////////////////////////
+template <typename FieldMeta, typename Type>
+struct FLogScope::FFieldSet<FieldMeta, Type[]>
+{
+	static void Impl(FLogScope*, Type const* Data, int32 Num)
+	{
+		static const uint32 Index = FieldMeta::Index & 0x7f;
+		int32 Size = (Num * sizeof(Type)) & (FAuxHeader::SizeLimit - 1) & ~(sizeof(Type) - 1);
+		Field_WriteAuxData(Index, (const uint8*)Data, Size);
+	}
+};
+
+#if STATICALLY_SIZED_ARRAY_FIELDS_SUPPORT
+////////////////////////////////////////////////////////////////////////////////
+template <typename FieldMeta, typename Type, int32 Count>
+struct FLogScope::FFieldSet<FieldMeta, Type[Count]>
+{
+	static void Impl(FLogScope*, Type const* Data, int32 Num=-1) = delete;
+};
+#endif // STATICALLY_SIZED_ARRAY_FIELDS_SUPPORT
+
+////////////////////////////////////////////////////////////////////////////////
+template <typename FieldMeta>
+struct FLogScope::FFieldSet<FieldMeta, AnsiString>
+{
+	static void Impl(FLogScope*, const ANSICHAR* String, int32 Length=-1)
+	{
+		if (Length < 0)
+		{
+			Length = int32(strlen(String));
+		}
+
+		static const uint32 Index = FieldMeta::Index & 0x7f;
+		Field_WriteStringAnsi(Index, String, Length);
+	}
+
+	static void Impl(FLogScope*, const TCHAR* String, int32 Length=-1)
+	{
+		if (Length < 0)
+		{
+			Length = 0;
+			for (const TCHAR* c = String; *c; ++c, ++Length);
+		}
+
+		static const uint32 Index = FieldMeta::Index & 0x7f;
+		Field_WriteStringAnsi(Index, String, Length);
+	}
+};
+
+////////////////////////////////////////////////////////////////////////////////
+template <typename FieldMeta>
+struct FLogScope::FFieldSet<FieldMeta, WideString>
+{
+	static void Impl(FLogScope*, const TCHAR* String, int32 Length=-1)
+	{
+		if (Length < 0)
+		{
+			Length = 0;
+			for (const TCHAR* c = String; *c; ++c, ++Length);
+		}
+
+		static const uint32 Index = FieldMeta::Index & 0x7f;
+		Field_WriteStringWide(Index, String, Length);
+	}
+};
+
+template <typename FieldMeta>
+struct FLogScope::FFieldSet<FieldMeta, Attachment>
+{
+	template <typename LambdaType>
+	static void Impl(FLogScope* Scope, LambdaType&& Lambda)
+	{
+		uint8* Dest = (uint8*)(Scope->Ptr) + FieldMeta::Offset;
+		Lambda(Dest);
+	}
+
+	static void Impl(FLogScope* Scope, const void* Data, uint32 Size)
+	{
+		uint8* Dest = (uint8*)(Scope->Ptr) + FieldMeta::Offset;
+		::memcpy(Dest, Data, Size);
+	}
+};
+
 } // namespace Private
 } // namespace Trace
 
diff --git a/Engine/Source/Runtime/TraceLog/Public/Trace/Detail/Trace.h b/Engine/Source/Runtime/TraceLog/Public/Trace/Detail/Trace.h
--- a/Engine/Source/Runtime/TraceLog/Public/Trace/Detail/Trace.h
+++ b/Engine/Source/Runtime/TraceLog/Public/Trace/Detail/Trace.h
@@ -6,10 +6,16 @@
 
 #if UE_TRACE_ENABLED
 
+#include <type_traits>
+
 namespace Trace
 {
-	class FChannel;
-};
+
+class FChannel;
+
+}; // namespace Trace
+
+#define TRACE_PRIVATE_STATISTICS (!UE_BUILD_TEST)
 
 #define TRACE_PRIVATE_CHANNEL_DEFAULT_ARGS false, "None"
 
@@ -61,10 +67,10 @@ namespace Trace
 		{ \
 			Important			= Trace::Private::FEventInfo::Flag_Important, \
 			NoSync				= Trace::Private::FEventInfo::Flag_NoSync, \
-			PartialEventFlags	= (0, ##__VA_ARGS__) & ~Important, \
+			PartialEventFlags	= (0, ##__VA_ARGS__), \
 		}; \
 		enum : bool { bIsImportant = ((0, ##__VA_ARGS__) & Important) != 0, }; \
-		static constexpr uint32 GetSize() { return decltype(EventProps_Private)::Size; } \
+		static constexpr uint32 GetSize() { return EventProps_Meta::Size; } \
 		static uint32 GetUid() { static uint32 Uid = 0; return (Uid = Uid ? Uid : Initialize()); } \
 		static uint32 FORCENOINLINE Initialize() \
 		{ \
@@ -84,27 +90,45 @@ namespace Trace
 			}(); \
 			return Uid_ThreadSafeInit; \
 		} \
-		Trace::TField<0 /*Index*/, 0 /*Offset*/,
+		typedef Trace::TField<0 /*Index*/, 0 /*Offset*/,
 
 #define TRACE_PRIVATE_EVENT_FIELD(FieldType, FieldName) \
-		FieldType> const FieldName##_Field = Trace::FLiteralName(#FieldName); \
-		template <typename... Ts> auto FieldName(Ts... ts) const { FieldName##_Field.Set((uint8*)this, Forward<Ts>(ts)...); return true; } \
-		Trace::TField< \
-			decltype(FieldName##_Field)::Index + 1, \
-			decltype(FieldName##_Field)::Offset + decltype(FieldName##_Field)::Size,
+		FieldType> FieldName##_Meta; \
+		FieldName##_Meta const FieldName##_Field = Trace::FLiteralName(#FieldName); \
+		template <typename... Ts> auto FieldName(Ts... ts) const { \
+			LogScopeType::FFieldSet<FieldName##_Meta, FieldType>::Impl((LogScopeType*)this, Forward<Ts>(ts)...); \
+			return true; \
+		} \
+		typedef Trace::TField< \
+			FieldName##_Meta::Index + 1, \
+			FieldName##_Meta::Offset + FieldName##_Meta::Size,
 
 #define TRACE_PRIVATE_EVENT_END() \
-		Trace::EventProps> const EventProps_Private = {}; \
-		Trace::TField<0, decltype(EventProps_Private)::Size, Trace::Attachment> const Attachment_Field = {}; \
-		template <typename... Ts> auto Attachment(Ts... ts) const { Attachment_Field.Set((uint8*)this, Forward<Ts>(ts)...); return true; } \
+		Trace::EventProps> EventProps_Meta; \
+		EventProps_Meta const EventProps_Private = {}; \
+		typedef Trace::TField<0, EventProps_Meta::Size, Trace::Attachment> Attachment_Meta; \
+		const Attachment_Meta Attachment_Field = {}; \
+		template <typename... Ts> auto Attachment(Ts... ts) const { \
+			AreAttachmentsAllowed<bIsImportant>(); \
+			LogScopeType::FFieldSet<Attachment_Meta, Trace::Attachment>::Impl((LogScopeType*)this, Forward<Ts>(ts)...); \
+			return true; \
+		} \
+		typedef std::conditional<bIsImportant, Trace::Private::FImportantLogScope, Trace::Private::FLogScope>::type LogScopeType; \
 		explicit operator bool () const { return true; } \
-		enum { EventFlags = PartialEventFlags|(decltype(EventProps_Private)::MaybeHasAux ? Trace::Private::FEventInfo::Flag_MaybeHasAux : 0), }; \
+		enum { EventFlags = PartialEventFlags|(EventProps_Meta::NumAuxFields ? Trace::Private::FEventInfo::Flag_MaybeHasAux : 0), }; \
+		template <bool bImportance> static void AreAttachmentsAllowed() { \
+			static_assert(!bImportance, "Important events cannot have attachments"); \
+		} \
+		static_assert( \
+			!bIsImportant || (EventFlags & Trace::Private::FEventInfo::Flag_NoSync), \
+			"Trace events flagged as Important events must be marked NoSync" \
+		); \
 	};
 
 #define TRACE_PRIVATE_LOG_PRELUDE(EnterFunc, LoggerName, EventName, ChannelsExpr, ...) \
 	if (TRACE_PRIVATE_CHANNELEXPR_IS_ENABLED(ChannelsExpr)) \
-		if (auto LogScope = Trace::Private::TLogScope<F##LoggerName##EventName##Fields>::EnterFunc(__VA_ARGS__)) \
-			if (const auto& __restrict EventName = *(F##LoggerName##EventName##Fields*)LogScope.GetPointer())
+		if (auto LogScope = F##LoggerName##EventName##Fields::LogScopeType::EnterFunc<F##LoggerName##EventName##Fields>(__VA_ARGS__)) \
+			if (const auto& __restrict EventName = *(F##LoggerName##EventName##Fields*)(&LogScope))
 
 #define TRACE_PRIVATE_LOG_EPILOG() \
 				LogScope += LogScope
diff --git a/Engine/Source/Runtime/TraceLog/Public/Trace/Trace.h b/Engine/Source/Runtime/TraceLog/Public/Trace/Trace.h
--- a/Engine/Source/Runtime/TraceLog/Public/Trace/Trace.h
+++ b/Engine/Source/Runtime/TraceLog/Public/Trace/Trace.h
@@ -26,12 +26,23 @@ struct FInitializeDesc
 	bool			bUseWorkerThread	= true;
 };
 
+struct FStatistics
+{
+	uint64		BytesSent;
+	uint64		BytesTraced;
+	uint32		MemoryUsed;
+	uint32		CacheUsed;
+	uint32		CacheWaste;
+};
+
 UE_TRACE_API void	Initialize(const FInitializeDesc& Desc) UE_TRACE_IMPL();
 UE_TRACE_API void	Shutdown() UE_TRACE_IMPL();
 UE_TRACE_API void	Update() UE_TRACE_IMPL();
+UE_TRACE_API void	GetStatistics(FStatistics& Out) UE_TRACE_IMPL();
 UE_TRACE_API bool	SendTo(const TCHAR* Host, uint32 Port=0) UE_TRACE_IMPL(false);
 UE_TRACE_API bool	WriteTo(const TCHAR* Path) UE_TRACE_IMPL(false);
 UE_TRACE_API bool	IsTracing() UE_TRACE_IMPL(false);
+UE_TRACE_API bool	Stop() UE_TRACE_IMPL(false);
 UE_TRACE_API bool	IsChannel(const TCHAR* ChanneName) UE_TRACE_IMPL(false);
 UE_TRACE_API bool	ToggleChannel(const TCHAR* ChannelName, bool bEnabled) UE_TRACE_IMPL(false);
 UE_TRACE_API void	ThreadRegister(const TCHAR* Name, uint32 SystemId, int32 SortHint) UE_TRACE_IMPL();
diff --git a/Engine/Source/Runtime/TraceLog/Public/Trace/Trace.inl b/Engine/Source/Runtime/TraceLog/Public/Trace/Trace.inl
--- a/Engine/Source/Runtime/TraceLog/Public/Trace/Trace.inl
+++ b/Engine/Source/Runtime/TraceLog/Public/Trace/Trace.inl
@@ -5,5 +5,6 @@
 #include "CoreTypes.h"
 #include "Detail/Channel.inl"
 #include "Detail/EventNode.h"
+#include "Detail/Important/ImportantLogScope.inl"
 #include "Detail/LogScope.inl"
 #include "Detail/Writer.inl"
