Fixed back-to-front #if/endif logic.

#rnx

diff --git a/Engine/Source/Runtime/TraceLog/Private/Trace/Writer.cpp b/Engine/Source/Runtime/TraceLog/Private/Trace/Writer.cpp
--- a/Engine/Source/Runtime/TraceLog/Private/Trace/Writer.cpp
+++ b/Engine/Source/Runtime/TraceLog/Private/Trace/Writer.cpp
@@ -149,7 +149,7 @@ void* Writer_MemoryAllocate(SIZE_T Size, uint32 Alignment)
 		Writer_SendData(ThreadId, TraceData.GetData(), TraceData.GetSize());
 	}
 
-#if !TRACE_PRIVATE_STATISTICS
+#if TRACE_PRIVATE_STATISTICS
 	AtomicIncrement(&GStatistics.MemoryUsed, Size);
 #endif
 
@@ -188,7 +188,7 @@ void Writer_MemoryFree(void* Address, uint32 Size)
 	}
 #endif // TRACE_PRIVATE_STOMP
 
-#if !TRACE_PRIVATE_STATISTICS
+#if TRACE_PRIVATE_STATISTICS
 	AtomicIncrement(&GStatistics.MemoryUsed, -int64(Size));
 #endif
 }
