Fixed fallout from back-to-front preprocessor statements.

#rnx

diff --git a/Engine/Source/Runtime/TraceLog/Private/Trace/Writer.cpp b/Engine/Source/Runtime/TraceLog/Private/Trace/Writer.cpp
--- a/Engine/Source/Runtime/TraceLog/Private/Trace/Writer.cpp
+++ b/Engine/Source/Runtime/TraceLog/Private/Trace/Writer.cpp
@@ -150,7 +150,7 @@ void* Writer_MemoryAllocate(SIZE_T Size, uint32 Alignment)
 	}
 
 #if TRACE_PRIVATE_STATISTICS
-	AtomicIncrement(&GStatistics.MemoryUsed, Size);
+	AtomicAddRelaxed(&GTraceStatistics.MemoryUsed, uint32(Size));
 #endif
 
 	return Ret;
@@ -189,7 +189,7 @@ void Writer_MemoryFree(void* Address, uint32 Size)
 #endif // TRACE_PRIVATE_STOMP
 
 #if TRACE_PRIVATE_STATISTICS
-	AtomicIncrement(&GStatistics.MemoryUsed, -int64(Size));
+	AtomicAddRelaxed(&GTraceStatistics.MemoryUsed, uint32(-int64(Size)));
 #endif
 }
 
