Fixed a warning from static analysis.

#rnx

diff --git a/Engine/Source/Runtime/TraceLog/Private/Trace/Important/SharedBuffer.cpp b/Engine/Source/Runtime/TraceLog/Private/Trace/Important/SharedBuffer.cpp
--- a/Engine/Source/Runtime/TraceLog/Private/Trace/Important/SharedBuffer.cpp
+++ b/Engine/Source/Runtime/TraceLog/Private/Trace/Important/SharedBuffer.cpp
@@ -167,23 +167,12 @@ void Writer_InitializeSharedBuffers()
 {
 	Writer_InitializeCache();
 
-	const uint32 BufferCount = 1;
+	FSharedBuffer* Buffer = Writer_CreateSharedBuffer();
 
-	FSharedBuffer* Buffers[BufferCount];
-	for (int i = 0; i < BufferCount; ++i)
-	{
-		Buffers[i] = Writer_CreateSharedBuffer();
-	}
-
-	for (int i = 1; i < BufferCount; ++i)
-	{
-		Buffers[i - 1]->Next = Buffers[i];
-	}
-
-	GTailBuffer = Buffers[0];
+	GTailBuffer = Buffer;
 	GTailPreSent = 0;
 
-	AtomicStoreRelease(&GSharedBuffer, Buffers[0]);
+	AtomicStoreRelease(&GSharedBuffer, Buffer);
 }
 
 ////////////////////////////////////////////////////////////////////////////////
