Fallback to a random(ish) control port if the known default was not available. This way we should always be able to communicate with running trace systems (provided we know where they are).

#rnx

diff --git a/Engine/Source/Runtime/TraceLog/Private/Trace/Control.cpp b/Engine/Source/Runtime/TraceLog/Private/Trace/Control.cpp
--- a/Engine/Source/Runtime/TraceLog/Private/Trace/Control.cpp
+++ b/Engine/Source/Runtime/TraceLog/Private/Trace/Control.cpp
@@ -101,6 +101,20 @@ static bool Writer_ControlDispatch(uint32 ArgC, ANSICHAR const* const* ArgV)
 static bool Writer_ControlListen()
 {
 	GControlListen = TcpSocketListen(1985);
+	if (!GControlListen)
+	{
+		uint32 Seed = uint32(TimeGetTimestamp());
+		for (uint32 i = 0; i < 10 && !GControlListen; Seed *= 13, ++i)
+		{
+			uint32 Port = (Seed & 0x1fff) + 0x8000;
+			GControlListen = TcpSocketListen(Port);
+			if (GControlListen)
+			{
+				break;
+			}
+		}
+	}
+
 	if (!GControlListen)
 	{
 		GControlState = EControlState::Failed;
