Accessor to fetch the port used to communicate with Trace

#rnx

diff --git a/Engine/Source/Runtime/TraceLog/Private/Trace/Control.cpp b/Engine/Source/Runtime/TraceLog/Private/Trace/Control.cpp
--- a/Engine/Source/Runtime/TraceLog/Private/Trace/Control.cpp
+++ b/Engine/Source/Runtime/TraceLog/Private/Trace/Control.cpp
@@ -47,6 +47,7 @@ static FControlCommands	GControlCommands;
 static UPTRINT			GControlListen		= 0;
 static UPTRINT			GControlSocket		= 0;
 static EControlState	GControlState;		// = EControlState::Closed;
+static uint32			GControlPort		= 1985;
 
 ////////////////////////////////////////////////////////////////////////////////
 static uint32 Writer_ControlHash(const ANSICHAR* Word)
@@ -100,7 +101,7 @@ static bool Writer_ControlDispatch(uint32 ArgC, ANSICHAR const* const* ArgV)
 ////////////////////////////////////////////////////////////////////////////////
 static bool Writer_ControlListen()
 {
-	GControlListen = TcpSocketListen(1985);
+	GControlListen = TcpSocketListen(GControlPort);
 	if (!GControlListen)
 	{
 		uint32 Seed = uint32(TimeGetTimestamp());
@@ -110,6 +111,7 @@ static bool Writer_ControlListen()
 			GControlListen = TcpSocketListen(Port);
 			if (GControlListen)
 			{
+				GControlPort = Port;
 				break;
 			}
 		}
@@ -239,6 +241,12 @@ static void Writer_ControlRecv()
 	}
 }
 
+////////////////////////////////////////////////////////////////////////////////
+uint32 Writer_GetControlPort()
+{
+	return GControlPort;
+}
+
 ////////////////////////////////////////////////////////////////////////////////
 void Writer_UpdateControl()
 {
