Update Trace's control server early so a port is assigned before we do anything else

#rnx

diff --git a/Engine/Source/Runtime/TraceLog/Private/Trace/Writer.cpp b/Engine/Source/Runtime/TraceLog/Private/Trace/Writer.cpp
--- a/Engine/Source/Runtime/TraceLog/Private/Trace/Writer.cpp
+++ b/Engine/Source/Runtime/TraceLog/Private/Trace/Writer.cpp
@@ -437,6 +437,8 @@ static void Writer_WorkerThread()
 {
 	Trace::ThreadRegister(TEXT("Trace"), 0, INT_MAX);
 
+	Writer_UpdateControl();
+
 	// At this point we haven't ever collected any trace events. So we'll stall
 	// for just a little bit to give the user a chance to set up sending the trace
 	// somewhere. This way they get all events since boot, otherwise they'll be
