Ports can content when multiple processes are tracing on a single machine. A known port is used as a discovery mechanism but when contention happens the OS is used to assign a port number to the control server. So that tools like Insights can talk to trace sessions it knows about, the control port is embedded in the trace stream's metadata.

#rnx

diff --git a/Engine/Source/Runtime/TraceLog/Private/Trace/Writer.cpp b/Engine/Source/Runtime/TraceLog/Private/Trace/Writer.cpp
--- a/Engine/Source/Runtime/TraceLog/Private/Trace/Writer.cpp
+++ b/Engine/Source/Runtime/TraceLog/Private/Trace/Writer.cpp
@@ -39,6 +39,7 @@ void			Writer_InitializePool();
 void			Writer_ShutdownPool();
 void			Writer_DrainBuffers();
 void			Writer_EndThreadBuffer();
+uint32			Writer_GetControlPort();
 void			Writer_UpdateControl();
 void			Writer_InitializeControl();
 void			Writer_ShutdownControl();
@@ -374,10 +375,13 @@ static bool Writer_UpdateConnection()
 	struct FHandshake
 	{
 		uint32 Magic			= 'TRC2';
-		uint16 MetadataSize		= uint16(0);
+		uint16 MetadataSize		= uint16(4); //  = sizeof(MetadataField0 + ControlPort)
+		uint16 MetadataField0	= uint16(sizeof(ControlPort) | (ControlPortFieldId << 8));
+		uint16 ControlPort		= uint16(Writer_GetControlPort());
 		enum
 		{
-			Size				= 6,
+			Size				= 10,
+			ControlPortFieldId	= 0,
 		};
 	};
 	FHandshake Handshake;
