There was a logic error when expanding a buffer's size if a trace event's required size fell between a shared buffer's usable size and the allocation block size.

#rnx

diff --git a/Engine/Source/Runtime/TraceLog/Private/Trace/Important/SharedBuffer.cpp b/Engine/Source/Runtime/TraceLog/Private/Trace/Important/SharedBuffer.cpp
--- a/Engine/Source/Runtime/TraceLog/Private/Trace/Important/SharedBuffer.cpp
+++ b/Engine/Source/Runtime/TraceLog/Private/Trace/Important/SharedBuffer.cpp
@@ -29,8 +29,10 @@ extern FStatistics		GTraceStatistics;
 ////////////////////////////////////////////////////////////////////////////////
 static FSharedBuffer* Writer_CreateSharedBuffer(uint32 SizeHint=0)
 {
+	const uint32 OverheadSize = sizeof(FSharedBuffer) + sizeof(uint32);
+
 	uint32 BlockSize = GBlockSize;
-	if (SizeHint > GBlockSize)
+	if (SizeHint > GBlockSize - OverheadSize)
 	{
 		BlockSize += SizeHint - GBlockSize;
 		BlockSize += sizeof(FSharedBuffer);
