Moved the Trace namespace into the UE namespace.

diff --git a/Engine/Source/Runtime/TraceLog/Public/Trace/Detail/Trace.h b/Engine/Source/Runtime/TraceLog/Public/Trace/Detail/Trace.h
--- a/Engine/Source/Runtime/TraceLog/Public/Trace/Detail/Trace.h
+++ b/Engine/Source/Runtime/TraceLog/Public/Trace/Detail/Trace.h
@@ -8,20 +8,21 @@
 
 #include <type_traits>
 
-namespace Trace
-{
+namespace UE {
+namespace Trace {
 
 class FChannel;
 
-}; // namespace Trace
+} // namespace Trace
+} // namespace UE
 
 #define TRACE_PRIVATE_STATISTICS (!UE_BUILD_TEST)
 
 #define TRACE_PRIVATE_CHANNEL_DEFAULT_ARGS false, "None"
 
 #define TRACE_PRIVATE_CHANNEL_DECLARE(LinkageType, ChannelName) \
-	static Trace::FChannel ChannelName##Object; \
-	LinkageType Trace::FChannel& ChannelName = ChannelName##Object;
+	static UE::Trace::FChannel ChannelName##Object; \
+	LinkageType UE::Trace::FChannel& ChannelName = ChannelName##Object;
 
 #define TRACE_PRIVATE_CHANNEL_IMPL(ChannelName, ...) \
 	struct F##ChannelName##Registrator \
@@ -45,13 +46,13 @@ class FChannel;
 	TRACE_PRIVATE_CHANNEL_IMPL(ChannelName, ##__VA_ARGS__)
 
 #define TRACE_PRIVATE_CHANNEL_EXTERN(ChannelName, ...) \
-	__VA_ARGS__ extern Trace::FChannel& ChannelName;
+	__VA_ARGS__ extern UE::Trace::FChannel& ChannelName;
 
 #define TRACE_PRIVATE_CHANNELEXPR_IS_ENABLED(ChannelsExpr) \
 	bool(ChannelsExpr)
 
 #define TRACE_PRIVATE_EVENT_DEFINE(LoggerName, EventName) \
-	Trace::Private::FEventNode LoggerName##EventName##Event;
+	UE::Trace::Private::FEventNode LoggerName##EventName##Event;
 
 #define TRACE_PRIVATE_EVENT_BEGIN(LoggerName, EventName, ...) \
 	TRACE_PRIVATE_EVENT_BEGIN_IMPL(static, LoggerName, EventName, ##__VA_ARGS__)
@@ -65,8 +66,8 @@ class FChannel;
 	{ \
 		enum \
 		{ \
-			Important			= Trace::Private::FEventInfo::Flag_Important, \
-			NoSync				= Trace::Private::FEventInfo::Flag_NoSync, \
+			Important			= UE::Trace::Private::FEventInfo::Flag_Important, \
+			NoSync				= UE::Trace::Private::FEventInfo::Flag_NoSync, \
 			PartialEventFlags	= (0, ##__VA_ARGS__), \
 		}; \
 		enum : bool { bIsImportant = ((0, ##__VA_ARGS__) & Important) != 0, }; \
@@ -76,7 +77,7 @@ class FChannel;
 		{ \
 			static const uint32 Uid_ThreadSafeInit = [] () \
 			{ \
-				using namespace Trace; \
+				using namespace UE::Trace; \
 				static F##LoggerName##EventName##Fields Fields; \
 				static Private::FEventInfo Info = \
 				{ \
@@ -90,37 +91,37 @@ class FChannel;
 			}(); \
 			return Uid_ThreadSafeInit; \
 		} \
-		typedef Trace::TField<0 /*Index*/, 0 /*Offset*/,
+		typedef UE::Trace::TField<0 /*Index*/, 0 /*Offset*/,
 
 #define TRACE_PRIVATE_EVENT_FIELD(FieldType, FieldName) \
 		FieldType> FieldName##_Meta; \
-		FieldName##_Meta const FieldName##_Field = Trace::FLiteralName(#FieldName); \
+		FieldName##_Meta const FieldName##_Field = UE::Trace::FLiteralName(#FieldName); \
 		template <typename... Ts> auto FieldName(Ts... ts) const { \
 			LogScopeType::FFieldSet<FieldName##_Meta, FieldType>::Impl((LogScopeType*)this, Forward<Ts>(ts)...); \
 			return true; \
 		} \
-		typedef Trace::TField< \
+		typedef UE::Trace::TField< \
 			FieldName##_Meta::Index + 1, \
 			FieldName##_Meta::Offset + FieldName##_Meta::Size,
 
 #define TRACE_PRIVATE_EVENT_END() \
-		Trace::EventProps> EventProps_Meta; \
+		UE::Trace::EventProps> EventProps_Meta; \
 		EventProps_Meta const EventProps_Private = {}; \
-		typedef Trace::TField<0, EventProps_Meta::Size, Trace::Attachment> Attachment_Meta; \
+		typedef UE::Trace::TField<0, EventProps_Meta::Size, UE::Trace::Attachment> Attachment_Meta; \
 		const Attachment_Meta Attachment_Field = {}; \
 		template <typename... Ts> auto Attachment(Ts... ts) const { \
 			AreAttachmentsAllowed<bIsImportant>(); \
-			LogScopeType::FFieldSet<Attachment_Meta, Trace::Attachment>::Impl((LogScopeType*)this, Forward<Ts>(ts)...); \
+			LogScopeType::FFieldSet<Attachment_Meta, UE::Trace::Attachment>::Impl((LogScopeType*)this, Forward<Ts>(ts)...); \
 			return true; \
 		} \
-		typedef std::conditional<bIsImportant, Trace::Private::FImportantLogScope, Trace::Private::FLogScope>::type LogScopeType; \
+		typedef std::conditional<bIsImportant, UE::Trace::Private::FImportantLogScope, UE::Trace::Private::FLogScope>::type LogScopeType; \
 		explicit operator bool () const { return true; } \
-		enum { EventFlags = PartialEventFlags|(EventProps_Meta::NumAuxFields ? Trace::Private::FEventInfo::Flag_MaybeHasAux : 0), }; \
+		enum { EventFlags = PartialEventFlags|(EventProps_Meta::NumAuxFields ? UE::Trace::Private::FEventInfo::Flag_MaybeHasAux : 0), }; \
 		template <bool bImportance> static void AreAttachmentsAllowed() { \
 			static_assert(!bImportance, "Important events cannot have attachments"); \
 		} \
 		static_assert( \
-			!bIsImportant || (EventFlags & Trace::Private::FEventInfo::Flag_NoSync), \
+			!bIsImportant || (EventFlags & UE::Trace::Private::FEventInfo::Flag_NoSync), \
 			"Trace events flagged as Important events must be marked NoSync" \
 		); \
 	};
@@ -138,13 +139,13 @@ class FChannel;
 		TRACE_PRIVATE_LOG_EPILOG()
 
 #define TRACE_PRIVATE_LOG_SCOPED(LoggerName, EventName, ChannelsExpr, ...) \
-	Trace::Private::FScopedLogScope PREPROCESSOR_JOIN(TheScope, __LINE__); \
+	UE::Trace::Private::FScopedLogScope PREPROCESSOR_JOIN(TheScope, __LINE__); \
 	TRACE_PRIVATE_LOG_PRELUDE(ScopedEnter, LoggerName, EventName, ChannelsExpr, ##__VA_ARGS__) \
 		PREPROCESSOR_JOIN(TheScope, __LINE__).SetActive(), \
 		TRACE_PRIVATE_LOG_EPILOG()
 
 #define TRACE_PRIVATE_LOG_SCOPED_T(LoggerName, EventName, ChannelsExpr, ...) \
-	Trace::Private::FScopedStampedLogScope PREPROCESSOR_JOIN(TheScope, __LINE__); \
+	UE::Trace::Private::FScopedStampedLogScope PREPROCESSOR_JOIN(TheScope, __LINE__); \
 	TRACE_PRIVATE_LOG_PRELUDE(ScopedStampedEnter, LoggerName, EventName, ChannelsExpr, ##__VA_ARGS__) \
 		PREPROCESSOR_JOIN(TheScope, __LINE__).SetActive(), \
 		TRACE_PRIVATE_LOG_EPILOG()
