Fix error C2872: 'Private': ambiguous symbol

diff --git a/Engine/Source/Runtime/TraceLog/Public/Trace/Detail/Trace.h b/Engine/Source/Runtime/TraceLog/Public/Trace/Detail/Trace.h
--- a/Engine/Source/Runtime/TraceLog/Public/Trace/Detail/Trace.h
+++ b/Engine/Source/Runtime/TraceLog/Public/Trace/Detail/Trace.h
@@ -79,7 +79,7 @@ class FChannel;
 			{ \
 				using namespace UE::Trace; \
 				static F##LoggerName##EventName##Fields Fields; \
-				static Private::FEventInfo Info = \
+				static UE::Trace::Private::FEventInfo Info = \
 				{ \
 					FLiteralName(#LoggerName), \
 					FLiteralName(#EventName), \
