Moved the joining of Trace's worker thread out into its own function

#rnx

diff --git a/Engine/Source/Runtime/TraceLog/Private/Trace/Writer.cpp b/Engine/Source/Runtime/TraceLog/Private/Trace/Writer.cpp
--- a/Engine/Source/Runtime/TraceLog/Private/Trace/Writer.cpp
+++ b/Engine/Source/Runtime/TraceLog/Private/Trace/Writer.cpp
@@ -482,6 +482,23 @@ static void Writer_WorkerCreate()
 	GWorkerThread = ThreadCreate("TraceWorker", Writer_WorkerThread);
 }
 
+////////////////////////////////////////////////////////////////////////////////
+static void Writer_WorkerJoin()
+{
+	if (!GWorkerThread)
+	{
+		return;
+	}
+
+	GWorkerThreadQuit = true;
+	ThreadJoin(GWorkerThread);
+	ThreadDestroy(GWorkerThread);
+
+	Writer_WorkerUpdate();
+
+	GWorkerThread = 0;
+}
+
 
 
 ////////////////////////////////////////////////////////////////////////////////
@@ -508,16 +525,7 @@ static void Writer_InternalShutdown()
 		return;
 	}
 
-	if (GWorkerThread)
-	{
-		GWorkerThreadQuit = true;
-		ThreadJoin(GWorkerThread);
-		ThreadDestroy(GWorkerThread);
-		GWorkerThread = 0;
-	}
-
-	Writer_WorkerUpdate();
-	Writer_DrainBuffers();
+	Writer_WorkerJoin();
 
 	if (GDataHandle)
 	{
