Don't fully shutdown Trace implicitly. Instead stop the worker thread (which Windows does in ExitProcess() and before static de-init anyway) and still allow the tracing of events. This prevents potential segfaults if any code in the atexit chain traces after, and captures a more complete trace.

#rnx

diff --git a/Engine/Source/Runtime/TraceLog/Private/Trace/Writer.cpp b/Engine/Source/Runtime/TraceLog/Private/Trace/Writer.cpp
--- a/Engine/Source/Runtime/TraceLog/Private/Trace/Writer.cpp
+++ b/Engine/Source/Runtime/TraceLog/Private/Trace/Writer.cpp
@@ -555,7 +555,12 @@ void Writer_InternalInitialize()
 			}
 			~FInitializer()
 			{
-				Writer_InternalShutdown();
+				/* We'll not shut anything down here so we can hopefully capture
+				 * any subsequent events. However, we will shutdown the worker
+				 * thread and leave it for something else to call update() (mem
+				 * tracing at time of writing). Windows will have already done
+				 * this implicitly in ExitProcess() anyway. */
+				Writer_WorkerJoin();
 			}
 		} Initializer;
 	}
