Fixed scoped-and-stamped trace events interleving with CpuProfilerTrace scopes incorrectly. The base timestamp can't be important as it is interpreted by analysis as if its come from the referenced thread. As important events are cached while threads come and go, important events can't have an association to a thread.

#rnx

diff --git a/Engine/Source/Runtime/TraceLog/Private/Trace/TlsBuffer.cpp b/Engine/Source/Runtime/TraceLog/Private/Trace/TlsBuffer.cpp
--- a/Engine/Source/Runtime/TraceLog/Private/Trace/TlsBuffer.cpp
+++ b/Engine/Source/Runtime/TraceLog/Private/Trace/TlsBuffer.cpp
@@ -23,7 +23,7 @@ extern uint64	GStartCycle;
 
 
 ////////////////////////////////////////////////////////////////////////////////
-UE_TRACE_EVENT_BEGIN($Trace, ThreadTiming, NoSync|Important)
+UE_TRACE_EVENT_BEGIN($Trace, ThreadTiming, NoSync)
 	UE_TRACE_EVENT_FIELD(uint64, BaseTimestamp)
 UE_TRACE_EVENT_END()
 
