Marked $Trace.* events as important as that is how they appear from an analysis perspective.

#rnx

diff --git a/Engine/Source/Runtime/TraceLog/Private/Trace/Writer.cpp b/Engine/Source/Runtime/TraceLog/Private/Trace/Writer.cpp
--- a/Engine/Source/Runtime/TraceLog/Private/Trace/Writer.cpp
+++ b/Engine/Source/Runtime/TraceLog/Private/Trace/Writer.cpp
@@ -48,7 +48,7 @@ void			Writer_ShutdownControl();
 
 
 ////////////////////////////////////////////////////////////////////////////////
-UE_TRACE_EVENT_BEGIN($Trace, NewTrace, NoSync)
+UE_TRACE_EVENT_BEGIN($Trace, NewTrace, NoSync|Important)
 	UE_TRACE_EVENT_FIELD(uint64, StartCycle)
 	UE_TRACE_EVENT_FIELD(uint64, CycleFrequency)
 	UE_TRACE_EVENT_FIELD(uint32, Serial)
@@ -58,7 +58,7 @@ UE_TRACE_EVENT_BEGIN($Trace, NewTrace, NoSync)
 	UE_TRACE_EVENT_FIELD(uint8, FeatureSet)
 UE_TRACE_EVENT_END()
 
-UE_TRACE_EVENT_BEGIN($Trace, SerialSync, NoSync)
+UE_TRACE_EVENT_BEGIN($Trace, SerialSync, NoSync|Important)
 UE_TRACE_EVENT_END()
 
 
