Removed an unused variable.

#rnx

diff --git a/Engine/Source/Runtime/TraceLog/Private/Trace/EventNode.cpp b/Engine/Source/Runtime/TraceLog/Private/Trace/EventNode.cpp
--- a/Engine/Source/Runtime/TraceLog/Private/Trace/EventNode.cpp
+++ b/Engine/Source/Runtime/TraceLog/Private/Trace/EventNode.cpp
@@ -123,8 +123,6 @@ void FEventNode::Describe() const
 	}
 
 	// Allocate the new event event in the log stream.
-	uint16 EventUid = EKnownEventUids::NewEvent << EKnownEventUids::_UidShift;
-
 	uint16 EventSize = sizeof(FNewEventEvent);
 	EventSize += sizeof(FNewEventEvent::Fields[0]) * Info->FieldCount;
 	EventSize += NamesSize;
