- Events that we want to capture to a buffer-redirect can't be marked as important (Important events don't use the TLS buffering path)
- The NewTrace event's UE_TRACE_LOG and Writer_SendData() were carefully ordered to make sure it is described in the stream before it is received. The SerialSync event however wasn't and the event description and traced event were in the wrong order. We'll dummy-trace both to force them to be described.
- The "SerialSync" event is used during analysis to indicate that no older events will be received and thus it is safe to derive a serial number. This scheme would fail if two successive calls to Writer_UpdateConnection() were made, which happened on boot.

#rnx

diff --git a/Engine/Source/Runtime/TraceLog/Private/Trace/Writer.cpp b/Engine/Source/Runtime/TraceLog/Private/Trace/Writer.cpp
--- a/Engine/Source/Runtime/TraceLog/Private/Trace/Writer.cpp
+++ b/Engine/Source/Runtime/TraceLog/Private/Trace/Writer.cpp
@@ -48,7 +48,7 @@ void			Writer_ShutdownControl();
 
 
 ////////////////////////////////////////////////////////////////////////////////
-UE_TRACE_EVENT_BEGIN($Trace, NewTrace, NoSync|Important)
+UE_TRACE_EVENT_BEGIN($Trace, NewTrace, NoSync)
 	UE_TRACE_EVENT_FIELD(uint64, StartCycle)
 	UE_TRACE_EVENT_FIELD(uint64, CycleFrequency)
 	UE_TRACE_EVENT_FIELD(uint32, Serial)
@@ -58,7 +58,7 @@ UE_TRACE_EVENT_BEGIN($Trace, NewTrace, NoSync|Important)
 	UE_TRACE_EVENT_FIELD(uint8, FeatureSet)
 UE_TRACE_EVENT_END()
 
-UE_TRACE_EVENT_BEGIN($Trace, SerialSync, NoSync|Important)
+UE_TRACE_EVENT_BEGIN($Trace, SerialSync, NoSync)
 UE_TRACE_EVENT_END()
 
 
@@ -375,17 +375,6 @@ static void Writer_LogHeader()
 ////////////////////////////////////////////////////////////////////////////////
 static bool Writer_UpdateConnection()
 {
-	if (GSerialSyncPending)
-	{
-		GSerialSyncPending = false;
-
-		// When analysis receives the "serial sync" event the starting serial
-		// can be established from preceeding synced events.
-		TWriteBufferRedirect<32> SideBuffer;
-		UE_TRACE_LOG($Trace, SerialSync, TraceLogChannel);
-		Writer_SendData(SideBuffer.GetData(), SideBuffer.GetSize());
-	}
-
 	if (!GPendingDataHandle)
 	{
 		return false;
@@ -449,14 +438,14 @@ static bool Writer_UpdateConnection()
 	GTraceStatistics.BytesSent = 0;
 	GTraceStatistics.BytesTraced = 0;
 
-	// Send the header events
-	TWriteBufferRedirect<512> HeaderEvents;
-	Writer_LogHeader();
-	HeaderEvents.Close();
-
+	// The first events we will send are ones that describe the trace's events
 	FEventNode::OnConnect();
 	Writer_DescribeEvents();
 
+	// Send the header event
+	TWriteBufferRedirect<512> HeaderEvents;
+	Writer_LogHeader();
+	HeaderEvents.Close();
 	Writer_SendData(HeaderEvents.GetData(), HeaderEvents.GetSize());
 
 	Writer_CacheOnConnect();
@@ -479,6 +468,17 @@ static void Writer_WorkerUpdate()
 	Writer_DescribeAnnounce();
 	Writer_UpdateSharedBuffers();
 	Writer_DrainBuffers();
+
+	if (GSerialSyncPending)
+	{
+		GSerialSyncPending = false;
+
+		// When analysis receives the "serial sync" event the starting serial
+		// can be established from preceding synced events.
+		TWriteBufferRedirect<32> SideBuffer;
+		UE_TRACE_LOG($Trace, SerialSync, TraceLogChannel);
+		Writer_SendData(SideBuffer.GetData(), SideBuffer.GetSize());
+	}
 }
 
 ////////////////////////////////////////////////////////////////////////////////
@@ -559,6 +559,16 @@ static void Writer_InternalInitializeImpl()
 	Writer_InitializeSharedBuffers();
 	Writer_InitializePool();
 	Writer_InitializeControl();
+
+	// The order of the events at the beginning of a trace is a little sensitive,
+	// especially for the two following events that are added to the trace stream
+	// at very specific points. By dummy-tracing them we invoke the code to IDs
+	// and can guarantee the events are described on connection. It's not pretty.
+	{
+		TWriteBufferRedirect<64> TempBuffer;
+		UE_TRACE_LOG($Trace, NewTrace, TraceLogChannel);
+		UE_TRACE_LOG($Trace, SerialSync, TraceLogChannel);
+	}
 }
 
 ////////////////////////////////////////////////////////////////////////////////
