Added exchange functions to Trace's atomics layer.

#rnx

diff --git a/Engine/Source/Runtime/TraceLog/Public/Trace/Detail/Atomic.h b/Engine/Source/Runtime/TraceLog/Public/Trace/Detail/Atomic.h
--- a/Engine/Source/Runtime/TraceLog/Public/Trace/Detail/Atomic.h
+++ b/Engine/Source/Runtime/TraceLog/Public/Trace/Detail/Atomic.h
@@ -17,6 +17,8 @@ template <typename Type> Type	AtomicLoadRelaxed(Type volatile* Source);
 template <typename Type> Type	AtomicLoadAcquire(Type volatile* Source);
 template <typename Type> void	AtomicStoreRelaxed(Type volatile* Target, Type Value);
 template <typename Type> void	AtomicStoreRelease(Type volatile* Target, Type Value);
+template <typename Type> Type	AtomicExchangeAcquire(Type volatile* Target, Type Value);
+template <typename Type> Type	AtomicExchangeRelease(Type volatile* Target, Type Value);
 template <typename Type> bool	AtomicCompareExchangeRelaxed(Type volatile* Target, Type New, Type Expected);
 template <typename Type> bool	AtomicCompareExchangeAcquire(Type volatile* Target, Type New, Type Expected);
 template <typename Type> bool	AtomicCompareExchangeRelease(Type volatile* Target, Type New, Type Expected);
@@ -73,6 +75,22 @@ inline void AtomicStoreRelease(Type volatile* Target, Type Value)
 	T->store(Value, std::memory_order_release);
 }
 
+////////////////////////////////////////////////////////////////////////////////
+template <typename Type>
+inline Type AtomicExchangeAcquire(Type volatile* Target, Type Value)
+{
+	std::atomic<Type>* T = (std::atomic<Type>*) Target;
+	return T->exchange(Value, std::memory_order_acquire);
+}
+
+////////////////////////////////////////////////////////////////////////////////
+template <typename Type>
+inline Type AtomicExchangeRelease(Type volatile* Target, Type Value)
+{
+	std::atomic<Type>* T = (std::atomic<Type>*) Target;
+	return T->exchange(Value, std::memory_order_release);
+}
+
 ////////////////////////////////////////////////////////////////////////////////
 template <typename Type>
 inline bool AtomicCompareExchangeRelaxed(Type volatile* Target, Type New, Type Expected)
