When a new trace connection is established known events need to be resent. A list of known events is maintained for this purpose (it predates the important event side-channel). If any new events were added to the new-events list, the on-connect logic would overwrite the list and all knowledge of the new events would be lost. Analysis would end up with events it knew nothing of and would be unable to proceed.

#rnx

diff --git a/Engine/Source/Runtime/TraceLog/Private/Trace/EventNode.cpp b/Engine/Source/Runtime/TraceLog/Private/Trace/EventNode.cpp
--- a/Engine/Source/Runtime/TraceLog/Private/Trace/EventNode.cpp
+++ b/Engine/Source/Runtime/TraceLog/Private/Trace/EventNode.cpp
@@ -192,6 +192,7 @@ void FEventNode::OnConnect()
 		FEventNode* Node = AtomicLoadRelaxed(&GNewEventList);
 		if (AtomicCompareExchangeRelaxed(&GNewEventList, GEventListHead, Node))
 		{
+			GEventListTail->Next = Node;
 			break;
 		}
 	}
