Use an exchange to acquire ownership of the new threads list.

#rnx

diff --git a/Engine/Source/Runtime/TraceLog/Private/Trace/TlsBuffer.cpp b/Engine/Source/Runtime/TraceLog/Private/Trace/TlsBuffer.cpp
--- a/Engine/Source/Runtime/TraceLog/Private/Trace/TlsBuffer.cpp
+++ b/Engine/Source/Runtime/TraceLog/Private/Trace/TlsBuffer.cpp
@@ -173,15 +173,7 @@ void Writer_DrainBuffers()
 #endif
 
 	// Claim ownership of any new thread buffer lists
-	FWriteBuffer* __restrict NewThreadList;
-	for (;; PlatformYield())
-	{
-		NewThreadList = AtomicLoadRelaxed(&GNewThreadList);
-		if (AtomicCompareExchangeAcquire(&GNewThreadList, (FWriteBuffer*)nullptr, NewThreadList))
-		{
-			break;
-		}
-	}
+	FWriteBuffer* __restrict NewThreadList = AtomicExchangeAcquire(&GNewThreadList, (FWriteBuffer*)nullptr);
 
 	// Reverse the new threads list so they're more closely ordered by age
 	// when sent out.
