Use exchange instead of CAS to gain ownership of the linked list of new trace events.

#rnx

diff --git a/Engine/Source/Runtime/TraceLog/Private/Trace/EventNode.cpp b/Engine/Source/Runtime/TraceLog/Private/Trace/EventNode.cpp
--- a/Engine/Source/Runtime/TraceLog/Private/Trace/EventNode.cpp
+++ b/Engine/Source/Runtime/TraceLog/Private/Trace/EventNode.cpp
@@ -40,18 +40,12 @@ const FEventNode* FEventNode::FIter::GetNext()
 ////////////////////////////////////////////////////////////////////////////////
 FEventNode::FIter FEventNode::ReadNew()
 {
-	FEventNode* EventList = AtomicLoadRelaxed(&GNewEventList);
+	FEventNode* EventList = AtomicExchangeAcquire(&GNewEventList, (FEventNode*)nullptr);
 	if (EventList == nullptr)
 	{
 		return {};
 	}
 
-	while (!AtomicCompareExchangeAcquire(&GNewEventList, (FEventNode*)nullptr, EventList))
-	{
-		PlatformYield();
-		EventList = AtomicLoadRelaxed(&GNewEventList);
-	}
-
 	if (GEventListHead == nullptr)
 	{
 		GEventListHead = EventList;
@@ -187,15 +181,7 @@ void FEventNode::OnConnect()
 		return;
 	}
 
-	for (;; PlatformYield())
-	{
-		FEventNode* Node = AtomicLoadRelaxed(&GNewEventList);
-		if (AtomicCompareExchangeRelaxed(&GNewEventList, GEventListHead, Node))
-		{
-			GEventListTail->Next = Node;
-			break;
-		}
-	}
+	GEventListTail->Next = AtomicExchangeAcquire(&GNewEventList, GEventListHead);
 
 	GEventListHead = GEventListTail = nullptr;
 }
