Removed the default initialiser. This was originally added to register threads correct but Trace no longer makes any assumptions about thread registration.

#rnx

diff --git a/Engine/Source/Runtime/TraceLog/Private/Trace/Writer.cpp b/Engine/Source/Runtime/TraceLog/Private/Trace/Writer.cpp
--- a/Engine/Source/Runtime/TraceLog/Private/Trace/Writer.cpp
+++ b/Engine/Source/Runtime/TraceLog/Private/Trace/Writer.cpp
@@ -644,17 +644,6 @@ void Writer_Update()
 	}
 }
 
-////////////////////////////////////////////////////////////////////////////////
-static bool const bEnsureDynamicInit = [] () -> bool
-{
-	// Trace will register the thread it is initialised on as the main thread. On
-	// the off chance that no events are trace during dynamic initialisation this
-	// lambda will get called to cover that scenario. Of course, this cunning plan
-	// may not work on some platforms or if TraceLog is loaded on another thread.
-	Writer_InternalInitialize();
-	return false;
-}();
-
 
 
 ////////////////////////////////////////////////////////////////////////////////
