Assert on the assumptions about the underlying transport of the trace stream. There is some code which assumes the format of packets of trace data. The change just ties them together should the transport be changed in the future.

#rnx

diff --git a/Engine/Source/Runtime/TraceLog/Private/Trace/Important/Cache.cpp b/Engine/Source/Runtime/TraceLog/Private/Trace/Important/Cache.cpp
--- a/Engine/Source/Runtime/TraceLog/Private/Trace/Important/Cache.cpp
+++ b/Engine/Source/Runtime/TraceLog/Private/Trace/Important/Cache.cpp
@@ -4,6 +4,7 @@
 
 #if UE_TRACE_ENABLED
 
+#include "Trace/Detail/Protocol.h"
 #include "Trace/Trace.h"
 
 #include <memory.h>
@@ -158,6 +159,8 @@ void Writer_InitializeCache()
 
 	GCacheActiveBuffer = Writer_CacheCreateBuffer(GCacheBufferSize);
 	GCacheActiveBuffer->TailNext = &GCacheHeadBuffer;
+
+	static_assert(ETransport::Active == ETransport::TidPacket, "The important cache is transport aware");
 }
 
 ////////////////////////////////////////////////////////////////////////////////
diff --git a/Engine/Source/Runtime/TraceLog/Private/Trace/Writer.cpp b/Engine/Source/Runtime/TraceLog/Private/Trace/Writer.cpp
--- a/Engine/Source/Runtime/TraceLog/Private/Trace/Writer.cpp
+++ b/Engine/Source/Runtime/TraceLog/Private/Trace/Writer.cpp
@@ -248,6 +248,8 @@ void Writer_SendDataRaw(const void* Data, uint32 Size)
 ////////////////////////////////////////////////////////////////////////////////
 uint32 Writer_SendData(uint32 ThreadId, uint8* __restrict Data, uint32 Size)
 {
+	static_assert(ETransport::Active == ETransport::TidPacket, "Active should be set to what the compiled code uses. It is used to track places that assume transport packet format");
+
 #if TRACE_PRIVATE_STATISTICS
 	GTraceStatistics.BytesTraced += Size;
 #endif
diff --git a/Engine/Source/Runtime/TraceLog/Public/Trace/Detail/Protocol.h b/Engine/Source/Runtime/TraceLog/Public/Trace/Detail/Protocol.h
--- a/Engine/Source/Runtime/TraceLog/Public/Trace/Detail/Protocol.h
+++ b/Engine/Source/Runtime/TraceLog/Public/Trace/Detail/Protocol.h
@@ -34,6 +34,7 @@ enum ETransport : uint8
 	Raw			= 1,
 	Packet		= 2,
 	TidPacket	= 3,
+	Active		= TidPacket,
 };
 
 enum ETransportTid : uint32
