Moved trace's transport enums into their own include file.

#rnx

diff --git a/Engine/Source/Developer/TraceAnalysis/Private/Analysis/Engine.cpp b/Engine/Source/Developer/TraceAnalysis/Private/Analysis/Engine.cpp
--- a/Engine/Source/Developer/TraceAnalysis/Private/Analysis/Engine.cpp
+++ b/Engine/Source/Developer/TraceAnalysis/Private/Analysis/Engine.cpp
@@ -11,6 +11,7 @@
 #include "Trace/Analysis.h"
 #include "Trace/Analyzer.h"
 #include "Trace/Detail/Protocol.h"
+#include "Trace/Detail/Transport.h"
 #include "Transport/PacketTransport.h"
 #include "Transport/Transport.h"
 #include "Transport/TidPacketTransport.h"
diff --git a/Engine/Source/Developer/TraceAnalysis/Private/Analysis/Transport/TidPacketTransport.cpp b/Engine/Source/Developer/TraceAnalysis/Private/Analysis/Transport/TidPacketTransport.cpp
--- a/Engine/Source/Developer/TraceAnalysis/Private/Analysis/Transport/TidPacketTransport.cpp
+++ b/Engine/Source/Developer/TraceAnalysis/Private/Analysis/Transport/TidPacketTransport.cpp
@@ -3,7 +3,7 @@
 #include "TidPacketTransport.h"
 #include "Algo/BinarySearch.h"
 #include "HAL/UnrealMemory.h"
-#include "Trace/Detail/Protocol.h"
+#include "Trace/Detail/Transport.h"
 
 ////////////////////////////////////////////////////////////////////////////////
 namespace UE {
diff --git a/Engine/Source/Runtime/TraceLog/Private/Trace/Important/Cache.cpp b/Engine/Source/Runtime/TraceLog/Private/Trace/Important/Cache.cpp
--- a/Engine/Source/Runtime/TraceLog/Private/Trace/Important/Cache.cpp
+++ b/Engine/Source/Runtime/TraceLog/Private/Trace/Important/Cache.cpp
@@ -5,6 +5,7 @@
 #if UE_TRACE_ENABLED
 
 #include "Trace/Detail/Protocol.h"
+#include "Trace/Detail/Transport.h"
 #include "Trace/Trace.h"
 
 #include <memory.h>
diff --git a/Engine/Source/Runtime/TraceLog/Private/Trace/Writer.cpp b/Engine/Source/Runtime/TraceLog/Private/Trace/Writer.cpp
--- a/Engine/Source/Runtime/TraceLog/Private/Trace/Writer.cpp
+++ b/Engine/Source/Runtime/TraceLog/Private/Trace/Writer.cpp
@@ -6,8 +6,9 @@
 
 #include "Platform.h"
 #include "Trace/Detail/Atomic.h"
-#include "Trace/Detail/Protocol.h"
 #include "Trace/Detail/Channel.h"
+#include "Trace/Detail/Protocol.h"
+#include "Trace/Detail/Transport.h"
 #include "Trace/Trace.inl"
 #include "WriteBufferRedirect.h"
 
diff --git a/Engine/Source/Runtime/TraceLog/Public/Trace/Detail/Protocol.h b/Engine/Source/Runtime/TraceLog/Public/Trace/Detail/Protocol.h
--- a/Engine/Source/Runtime/TraceLog/Public/Trace/Detail/Protocol.h
+++ b/Engine/Source/Runtime/TraceLog/Public/Trace/Detail/Protocol.h
@@ -23,25 +23,3 @@
 #if defined(_MSC_VER)
 	#pragma warning(pop)
 #endif
-
-////////////////////////////////////////////////////////////////////////////////
-namespace UE {
-namespace Trace {
-
-enum ETransport : uint8
-{
-	_Unused		= 0,
-	Raw			= 1,
-	Packet		= 2,
-	TidPacket	= 3,
-	Active		= TidPacket,
-};
-
-enum ETransportTid : uint32
-{
-	Internal	= 0,
-	Bias		= 1,
-};
-
-} // namespace Trace
-} // namespace UE
diff --git a/Engine/Source/Runtime/TraceLog/Public/Trace/Detail/Transport.h b/Engine/Source/Runtime/TraceLog/Public/Trace/Detail/Transport.h
new file mode 100644
--- /dev/null
+++ b/Engine/Source/Runtime/TraceLog/Public/Trace/Detail/Transport.h
@@ -0,0 +1,26 @@
+// Copyright Epic Games, Inc. All Rights Reserved.
+
+#pragma once
+
+namespace UE {
+namespace Trace {
+
+////////////////////////////////////////////////////////////////////////////////
+enum ETransport : uint8
+{
+	_Unused		= 0,
+	Raw			= 1,
+	Packet		= 2,
+	TidPacket	= 3,
+	Active		= TidPacket,
+};
+
+////////////////////////////////////////////////////////////////////////////////
+enum ETransportTid : uint32
+{
+	Internal	= 0,
+	Bias		= 1,
+};
+
+} // namespace Trace
+} // namespace UE
