Be explicit about what the 0 means when building cached trace packets.

#rnx

diff --git a/Engine/Source/Runtime/TraceLog/Private/Trace/Important/Cache.cpp b/Engine/Source/Runtime/TraceLog/Private/Trace/Important/Cache.cpp
--- a/Engine/Source/Runtime/TraceLog/Private/Trace/Important/Cache.cpp
+++ b/Engine/Source/Runtime/TraceLog/Private/Trace/Important/Cache.cpp
@@ -93,7 +93,7 @@ static void Writer_CacheCommit(const FCacheBuffer* Collector)
 	uint32 OutputSize = Encode(Collector->Data, InputSize, Packet->Data, EncodeMaxSize);
 
 	Packet->PacketSize = OutputSize + sizeof(FPacketEncoded);
-	Packet->ThreadId = 0x8000 | uint16(0);
+	Packet->ThreadId = 0x8000 | uint16(ETransportTid::Internal);
 	Packet->DecodedSize = uint16(InputSize);
 
 	Used = sizeof(FPacketEncoded) + OutputSize;
