Removed the implicit sending of data on Trace's "internal" thread in favour of being explicit where and when data should be sent in that manner.

#rnx

diff --git a/Engine/Source/Runtime/TraceLog/Private/Trace/Important/Cache.cpp b/Engine/Source/Runtime/TraceLog/Private/Trace/Important/Cache.cpp
--- a/Engine/Source/Runtime/TraceLog/Private/Trace/Important/Cache.cpp
+++ b/Engine/Source/Runtime/TraceLog/Private/Trace/Important/Cache.cpp
@@ -20,7 +20,7 @@ int32	Encode(const void*, int32, void*, int32);
 void*	Writer_MemoryAllocate(SIZE_T, uint32);
 void	Writer_MemoryFree(void*, uint32);
 void	Writer_SendDataRaw(const void*, uint32);
-uint32	Writer_SendData(uint8* __restrict, uint32);
+uint32	Writer_SendData(uint32, uint8* __restrict, uint32);
 
 ////////////////////////////////////////////////////////////////////////////////
 struct alignas(16) FCacheBuffer
@@ -107,7 +107,7 @@ static void Writer_CacheCommit(const FCacheBuffer* Collector)
 ////////////////////////////////////////////////////////////////////////////////
 void Writer_CacheData(uint8* Data, uint32 Size)
 {
-	Writer_SendData(Data, Size);
+	Writer_SendData(ETransportTid::Internal, Data, Size);
 
 	while (true)
 	{
@@ -149,7 +149,7 @@ void Writer_CacheOnConnect()
 
 	if (uint32 Used = GCacheCollector->Size - GCacheCollector->Remaining)
 	{
-		Writer_SendData(GCacheCollector->Data, Used);
+		Writer_SendData(ETransportTid::Internal, GCacheCollector->Data, Used);
 	}
 }
 
diff --git a/Engine/Source/Runtime/TraceLog/Private/Trace/Writer.cpp b/Engine/Source/Runtime/TraceLog/Private/Trace/Writer.cpp
--- a/Engine/Source/Runtime/TraceLog/Private/Trace/Writer.cpp
+++ b/Engine/Source/Runtime/TraceLog/Private/Trace/Writer.cpp
@@ -308,12 +308,6 @@ uint32 Writer_SendData(uint32 ThreadId, uint8* __restrict Data, uint32 Size)
 	return Packet.PacketSize;
 }
 
-////////////////////////////////////////////////////////////////////////////////
-uint32 Writer_SendData(uint8* __restrict Data, uint32 Size)
-{
-	return Writer_SendData(ETransportTid::Internal, Data, Size);
-}
-
 ////////////////////////////////////////////////////////////////////////////////
 static void Writer_DescribeEvents()
 {
@@ -327,14 +321,14 @@ static void Writer_DescribeEvents()
 		// Flush just in case an NewEvent event will be larger than 512 bytes.
 		if (TraceData.GetSize() >= (TraceData.GetCapacity() - 512))
 		{
-			Writer_SendData(TraceData.GetData(), TraceData.GetSize());
+			Writer_SendData(ETransportTid::Internal, TraceData.GetData(), TraceData.GetSize());
 			TraceData.Reset();
 		}
 	}
 
 	if (TraceData.GetSize())
 	{
-		Writer_SendData(TraceData.GetData(), TraceData.GetSize());
+		Writer_SendData(ETransportTid::Internal, TraceData.GetData(), TraceData.GetSize());
 	}
 }
 
@@ -449,7 +443,7 @@ static bool Writer_UpdateConnection()
 	TWriteBufferRedirect<512> HeaderEvents;
 	Writer_LogHeader();
 	HeaderEvents.Close();
-	Writer_SendData(HeaderEvents.GetData(), HeaderEvents.GetSize());
+	Writer_SendData(ETransportTid::Internal, HeaderEvents.GetData(), HeaderEvents.GetSize());
 
 	Writer_CacheOnConnect();
 
@@ -480,7 +474,7 @@ static void Writer_WorkerUpdate()
 		// can be established from preceding synced events.
 		TWriteBufferRedirect<32> SideBuffer;
 		UE_TRACE_LOG($Trace, SerialSync, TraceLogChannel);
-		Writer_SendData(SideBuffer.GetData(), SideBuffer.GetSize());
+		Writer_SendData(ETransportTid::Internal, SideBuffer.GetData(), SideBuffer.GetSize());
 	}
 }
 
