Consolidated some of the duplication around building trace transport packets.

#rnx

diff --git a/Engine/Source/Runtime/TraceLog/Private/Trace/Important/Cache.cpp b/Engine/Source/Runtime/TraceLog/Private/Trace/Important/Cache.cpp
--- a/Engine/Source/Runtime/TraceLog/Private/Trace/Important/Cache.cpp
+++ b/Engine/Source/Runtime/TraceLog/Private/Trace/Important/Cache.cpp
@@ -61,18 +61,10 @@ static FCacheBuffer* Writer_CacheCreateBuffer(uint32 Size)
 ////////////////////////////////////////////////////////////////////////////////
 static void Writer_CacheCommit(const FCacheBuffer* Collector)
 {
-	struct FPacketEncoded
-	{
-		uint16	PacketSize;
-		uint16	ThreadId;
-		uint16	DecodedSize;
-		uint8	Data[];
-	};
-
 	// Check there's enough size to compress Data into.
 	uint32 InputSize = uint32(Collector->Size - Collector->Remaining);
 	uint32 EncodeMaxSize = GetEncodeMaxSize(InputSize);
-	if (EncodeMaxSize + sizeof(FPacketEncoded) > GCacheActiveBuffer->Remaining)
+	if (EncodeMaxSize + sizeof(FTidPacketEncoded) > GCacheActiveBuffer->Remaining)
 	{
 		// Retire active buffer
 		*(GCacheActiveBuffer->TailNext) = GCacheActiveBuffer;
@@ -89,14 +81,14 @@ static void Writer_CacheCommit(const FCacheBuffer* Collector)
 	}
 
 	uint32 Used = GCacheActiveBuffer->Size - GCacheActiveBuffer->Remaining;
-	auto* Packet = (FPacketEncoded*)(GCacheActiveBuffer->Data + Used);
+	auto* Packet = (FTidPacketEncoded*)(GCacheActiveBuffer->Data + Used);
 	uint32 OutputSize = Encode(Collector->Data, InputSize, Packet->Data, EncodeMaxSize);
 
-	Packet->PacketSize = OutputSize + sizeof(FPacketEncoded);
+	Packet->PacketSize = OutputSize + sizeof(FTidPacketEncoded);
 	Packet->ThreadId = 0x8000 | uint16(ETransportTid::Internal);
 	Packet->DecodedSize = uint16(InputSize);
 
-	Used = sizeof(FPacketEncoded) + OutputSize;
+	Used = sizeof(FTidPacketEncoded) + OutputSize;
 	GCacheActiveBuffer->Remaining -= Used;
 
 #if TRACE_PRIVATE_STATISTICS
diff --git a/Engine/Source/Runtime/TraceLog/Private/Trace/Writer.cpp b/Engine/Source/Runtime/TraceLog/Private/Trace/Writer.cpp
--- a/Engine/Source/Runtime/TraceLog/Private/Trace/Writer.cpp
+++ b/Engine/Source/Runtime/TraceLog/Private/Trace/Writer.cpp
@@ -260,20 +260,13 @@ uint32 Writer_SendData(uint32 ThreadId, uint8* __restrict Data, uint32 Size)
 		return 0;
 	}
 
-	struct FPacketBase
-	{
-		uint16 PacketSize;
-		uint16 ThreadId;
-	};
-
 	// Smaller buffers usually aren't redundant enough to benefit from being
 	// compressed. They often end up being larger.
 	if (Size <= 384)
 	{
-		static_assert(sizeof(FPacketBase) == sizeof(uint32), "");
-		Data -= sizeof(FPacketBase);
-		Size += sizeof(FPacketBase);
-		auto* Packet = (FPacketBase*)Data;
+		Data -= sizeof(FTidPacket);
+		Size += sizeof(FTidPacket);
+		auto* Packet = (FTidPacket*)Data;
 		Packet->ThreadId = uint16(ThreadId & 0x7fff);
 		Packet->PacketSize = uint16(Size);
 
@@ -282,26 +275,15 @@ uint32 Writer_SendData(uint32 ThreadId, uint8* __restrict Data, uint32 Size)
 		return Size;
 	}
 
-	struct FPacketEncoded
-		: public FPacketBase
-	{
-		uint16	DecodedSize;
-	};
-
-	struct FPacket
-		: public FPacketEncoded
-	{
-		// Buffer size is expressed as "A + B" where A is a maximum expected
-		// input size (i.e. at least GPoolBlockSize) and B is LZ4 overhead as
-		// per LZ4_COMPRESSBOUND.
-		uint8 Data[8129 + 64];
-	};
+	// Buffer size is expressed as "A + B" where A is a maximum expected
+	// input size (i.e. at least GPoolBlockSize) and B is LZ4 overhead as
+	// per LZ4_COMPRESSBOUND.
+	TTidPacketEncoded<8192 + 64> Packet;
 
-	FPacket Packet;
 	Packet.ThreadId = 0x8000 | uint16(ThreadId & 0x7fff);
 	Packet.DecodedSize = uint16(Size);
 	Packet.PacketSize = Encode(Data, Packet.DecodedSize, Packet.Data, sizeof(Packet.Data));
-	Packet.PacketSize += sizeof(FPacketEncoded);
+	Packet.PacketSize += sizeof(FTidPacketEncoded);
 
 	Writer_SendDataRaw(&Packet, Packet.PacketSize);
 
diff --git a/Engine/Source/Runtime/TraceLog/Public/Trace/Detail/Transport.h b/Engine/Source/Runtime/TraceLog/Public/Trace/Detail/Transport.h
--- a/Engine/Source/Runtime/TraceLog/Public/Trace/Detail/Transport.h
+++ b/Engine/Source/Runtime/TraceLog/Public/Trace/Detail/Transport.h
@@ -22,5 +22,45 @@ enum ETransportTid : uint32
 	Bias		= 1,
 };
 
+
+
+namespace Private
+{
+
+////////////////////////////////////////////////////////////////////////////////
+struct FTidPacketBase
+{
+	enum : uint16 { EncodedMarker = 0x8000 };
+
+	uint16 PacketSize;
+	uint16 ThreadId;
+};
+
+template <uint32 DataSize>
+struct TTidPacket
+	: public FTidPacketBase
+{
+	uint8	Data[DataSize];
+};
+
+template <uint32 DataSize>
+struct TTidPacketEncoded
+	: public FTidPacketBase
+{
+	uint16	DecodedSize;
+	uint8	Data[DataSize];
+};
+
+using FTidPacket		= TTidPacket<0>;
+using FTidPacketEncoded = TTidPacketEncoded<0>;
+
+////////////////////////////////////////////////////////////////////////////////
+// Some assumptions are made about 0-sized arrays in the packet structs so we
+// will casually make assertions about those assumptions here.
+static_assert(sizeof(FTidPacket) == 4, "");
+static_assert(sizeof(FTidPacketEncoded) == 6, "");
+
+} // namespace Private
+
 } // namespace Trace
 } // namespace UE
