Send event descriptions out before important events. Events and important events are maintained in two separate data structures. The visibility of one with respect to the other is arbitrary from the perspective of Trace's worker thread. This could result in event descriptions being after one of their events - unrecoverable when both are in the same packet.

#rnx

diff --git a/Engine/Source/Developer/TraceAnalysis/Private/Analysis/Transport/TidPacketTransport.h b/Engine/Source/Developer/TraceAnalysis/Private/Analysis/Transport/TidPacketTransport.h
--- a/Engine/Source/Developer/TraceAnalysis/Private/Analysis/Transport/TidPacketTransport.h
+++ b/Engine/Source/Developer/TraceAnalysis/Private/Analysis/Transport/TidPacketTransport.h
@@ -34,6 +34,7 @@ private:
 	bool					ReadPacket();
 	FThreadStream&			FindOrAddThread(uint32 ThreadId);
 	TArray<FThreadStream>	Threads = {
+								{ {}, ETransportTid::Events },
 								{ {}, ETransportTid::Internal },
 							};
 };
diff --git a/Engine/Source/Runtime/TraceLog/Private/Trace/Writer.cpp b/Engine/Source/Runtime/TraceLog/Private/Trace/Writer.cpp
--- a/Engine/Source/Runtime/TraceLog/Private/Trace/Writer.cpp
+++ b/Engine/Source/Runtime/TraceLog/Private/Trace/Writer.cpp
@@ -303,14 +303,14 @@ static void Writer_DescribeEvents()
 		// Flush just in case an NewEvent event will be larger than 512 bytes.
 		if (TraceData.GetSize() >= (TraceData.GetCapacity() - 512))
 		{
-			Writer_SendData(ETransportTid::Internal, TraceData.GetData(), TraceData.GetSize());
+			Writer_SendData(ETransportTid::Events, TraceData.GetData(), TraceData.GetSize());
 			TraceData.Reset();
 		}
 	}
 
 	if (TraceData.GetSize())
 	{
-		Writer_SendData(ETransportTid::Internal, TraceData.GetData(), TraceData.GetSize());
+		Writer_SendData(ETransportTid::Events, TraceData.GetData(), TraceData.GetSize());
 	}
 }
 
diff --git a/Engine/Source/Runtime/TraceLog/Public/Trace/Detail/Transport.h b/Engine/Source/Runtime/TraceLog/Public/Trace/Detail/Transport.h
--- a/Engine/Source/Runtime/TraceLog/Public/Trace/Detail/Transport.h
+++ b/Engine/Source/Runtime/TraceLog/Public/Trace/Detail/Transport.h
@@ -18,8 +18,9 @@ enum ETransport : uint8
 ////////////////////////////////////////////////////////////////////////////////
 enum ETransportTid : uint32
 {
-	Internal	= 0,
-	Bias		= 1,
+	Events		= 0,
+	Internal	= 1,
+	Bias		= 2,
 };
 
 
