Identify which packets of events are important and which are internal to Trace's operation. They can still be the same internal "thread" as that level of separation isn't needed, but it is helpful from a maintenance perspective to semantically separate them. It is plural "Importants" to placate Clang.

#rnx

diff --git a/Engine/Source/Developer/TraceAnalysis/Private/Analysis/Transport/TidPacketTransport.cpp b/Engine/Source/Developer/TraceAnalysis/Private/Analysis/Transport/TidPacketTransport.cpp
--- a/Engine/Source/Developer/TraceAnalysis/Private/Analysis/Transport/TidPacketTransport.cpp
+++ b/Engine/Source/Developer/TraceAnalysis/Private/Analysis/Transport/TidPacketTransport.cpp
@@ -87,7 +87,7 @@ void FTidPacketTransport::Update()
 
 	Threads.RemoveAll([] (const FThreadStream& Thread)
 	{
-		return (Thread.ThreadId <= ETransportTid::Internal) ? false : Thread.Buffer.IsEmpty();
+		return (Thread.ThreadId <= ETransportTid::Importants) ? false : Thread.Buffer.IsEmpty();
 	});
 }
 
diff --git a/Engine/Source/Developer/TraceAnalysis/Private/Analysis/Transport/TidPacketTransport.h b/Engine/Source/Developer/TraceAnalysis/Private/Analysis/Transport/TidPacketTransport.h
--- a/Engine/Source/Developer/TraceAnalysis/Private/Analysis/Transport/TidPacketTransport.h
+++ b/Engine/Source/Developer/TraceAnalysis/Private/Analysis/Transport/TidPacketTransport.h
@@ -35,7 +35,7 @@ private:
 	FThreadStream&			FindOrAddThread(uint32 ThreadId);
 	TArray<FThreadStream>	Threads = {
 								{ {}, ETransportTid::Events },
-								{ {}, ETransportTid::Internal },
+								{ {}, ETransportTid::Importants },
 							};
 };
 
diff --git a/Engine/Source/Runtime/TraceLog/Private/Trace/Important/Cache.cpp b/Engine/Source/Runtime/TraceLog/Private/Trace/Important/Cache.cpp
--- a/Engine/Source/Runtime/TraceLog/Private/Trace/Important/Cache.cpp
+++ b/Engine/Source/Runtime/TraceLog/Private/Trace/Important/Cache.cpp
@@ -85,7 +85,7 @@ static void Writer_CacheCommit(const FCacheBuffer* Collector)
 	uint32 OutputSize = Encode(Collector->Data, InputSize, Packet->Data, EncodeMaxSize);
 
 	Packet->PacketSize = OutputSize + sizeof(FTidPacketEncoded);
-	Packet->ThreadId = 0x8000 | uint16(ETransportTid::Internal);
+	Packet->ThreadId = 0x8000 | uint16(ETransportTid::Importants);
 	Packet->DecodedSize = uint16(InputSize);
 
 	Used = sizeof(FTidPacketEncoded) + OutputSize;
@@ -99,7 +99,7 @@ static void Writer_CacheCommit(const FCacheBuffer* Collector)
 ////////////////////////////////////////////////////////////////////////////////
 void Writer_CacheData(uint8* Data, uint32 Size)
 {
-	Writer_SendData(ETransportTid::Internal, Data, Size);
+	Writer_SendData(ETransportTid::Importants, Data, Size);
 
 	while (true)
 	{
@@ -141,7 +141,7 @@ void Writer_CacheOnConnect()
 
 	if (uint32 Used = GCacheCollector->Size - GCacheCollector->Remaining)
 	{
-		Writer_SendData(ETransportTid::Internal, GCacheCollector->Data, Used);
+		Writer_SendData(ETransportTid::Importants, GCacheCollector->Data, Used);
 	}
 }
 
diff --git a/Engine/Source/Runtime/TraceLog/Public/Trace/Detail/Transport.h b/Engine/Source/Runtime/TraceLog/Public/Trace/Detail/Transport.h
--- a/Engine/Source/Runtime/TraceLog/Public/Trace/Detail/Transport.h
+++ b/Engine/Source/Runtime/TraceLog/Public/Trace/Detail/Transport.h
@@ -18,9 +18,10 @@ enum ETransport : uint8
 ////////////////////////////////////////////////////////////////////////////////
 enum ETransportTid : uint32
 {
-	Events		= 0,
-	Internal	= 1,
-	Bias		= 2,
+	Events		= 0,			// used to describe events
+	Internal	= 1,			// events to make the trace stream function
+	Importants	= Internal,		// important/cached events
+	Bias,
 };
 
 
