Do not try to send any prepared packets of trace data if there is no where to send the data to.

#rnx

diff --git a/Engine/Source/Runtime/TraceLog/Private/Trace/Writer.cpp b/Engine/Source/Runtime/TraceLog/Private/Trace/Writer.cpp
--- a/Engine/Source/Runtime/TraceLog/Private/Trace/Writer.cpp
+++ b/Engine/Source/Runtime/TraceLog/Private/Trace/Writer.cpp
@@ -233,7 +233,7 @@ static bool						GSerialSyncPending;	// = false
 UPTRINT							GPendingDataHandle;	// = 0
 
 ////////////////////////////////////////////////////////////////////////////////
-void Writer_SendDataRaw(const void* Data, uint32 Size)
+static void Writer_SendDataImpl(const void* Data, uint32 Size)
 {
 #if TRACE_PRIVATE_STATISTICS
 	GTraceStatistics.BytesSent += Size;
@@ -246,6 +246,17 @@ void Writer_SendDataRaw(const void* Data, uint32 Size)
 	}
 }
 
+////////////////////////////////////////////////////////////////////////////////
+void Writer_SendDataRaw(const void* Data, uint32 Size)
+{
+	if (!GDataHandle)
+	{
+		return;
+	}
+
+	Writer_SendDataImpl(Data, Size);
+}
+
 ////////////////////////////////////////////////////////////////////////////////
 uint32 Writer_SendData(uint32 ThreadId, uint8* __restrict Data, uint32 Size)
 {
@@ -270,7 +281,7 @@ uint32 Writer_SendData(uint32 ThreadId, uint8* __restrict Data, uint32 Size)
 		Packet->ThreadId = uint16(ThreadId & 0x7fff);
 		Packet->PacketSize = uint16(Size);
 
-		Writer_SendDataRaw(Data, Size);
+		Writer_SendDataImpl(Data, Size);
 
 		return Size;
 	}
@@ -285,7 +296,7 @@ uint32 Writer_SendData(uint32 ThreadId, uint8* __restrict Data, uint32 Size)
 	Packet.PacketSize = Encode(Data, Packet.DecodedSize, Packet.Data, sizeof(Packet.Data));
 	Packet.PacketSize += sizeof(FTidPacketEncoded);
 
-	Writer_SendDataRaw(&Packet, Packet.PacketSize);
+	Writer_SendDataImpl(&Packet, Packet.PacketSize);
 
 	return Packet.PacketSize;
 }
