Added a static_assert() to confirm assumptions about the FControlCommands struct.

#rnx

diff --git a/Engine/Source/Runtime/TraceLog/Private/Trace/Control.cpp b/Engine/Source/Runtime/TraceLog/Private/Trace/Control.cpp
--- a/Engine/Source/Runtime/TraceLog/Private/Trace/Control.cpp
+++ b/Engine/Source/Runtime/TraceLog/Private/Trace/Control.cpp
@@ -9,6 +9,8 @@
 #include "Misc/CString.h"
 #include "Templates/UnrealTemplate.h"
 
+#include <type_traits>
+
 namespace UE {
 namespace Trace {
 namespace Private {
@@ -41,6 +43,7 @@ struct FControlCommands
 	}			Commands[Max];
 	uint8		Count;
 };
+static_assert(std::is_trivial<FControlCommands>(), "FControlCommands must be trivial");
 
 
 
