Do not return the number of bytes sent. It is no longer used anywhere

#rnx

diff --git a/Engine/Source/Runtime/TraceLog/Private/Trace/Important/Cache.cpp b/Engine/Source/Runtime/TraceLog/Private/Trace/Important/Cache.cpp
--- a/Engine/Source/Runtime/TraceLog/Private/Trace/Important/Cache.cpp
+++ b/Engine/Source/Runtime/TraceLog/Private/Trace/Important/Cache.cpp
@@ -20,7 +20,7 @@ int32	Encode(const void*, int32, void*, int32);
 void*	Writer_MemoryAllocate(SIZE_T, uint32);
 void	Writer_MemoryFree(void*, uint32);
 void	Writer_SendDataRaw(const void*, uint32);
-uint32	Writer_SendData(uint32, uint8* __restrict, uint32);
+void	Writer_SendData(uint32, uint8* __restrict, uint32);
 
 ////////////////////////////////////////////////////////////////////////////////
 struct alignas(16) FCacheBuffer
diff --git a/Engine/Source/Runtime/TraceLog/Private/Trace/TlsBuffer.cpp b/Engine/Source/Runtime/TraceLog/Private/Trace/TlsBuffer.cpp
--- a/Engine/Source/Runtime/TraceLog/Private/Trace/TlsBuffer.cpp
+++ b/Engine/Source/Runtime/TraceLog/Private/Trace/TlsBuffer.cpp
@@ -14,7 +14,7 @@ namespace Trace {
 namespace Private {
 
 ////////////////////////////////////////////////////////////////////////////////
-uint32			Writer_SendData(uint32, uint8* __restrict, uint32);
+void			Writer_SendData(uint32, uint8* __restrict, uint32);
 FWriteBuffer*	Writer_AllocateBlockFromPool();
 uint32			Writer_GetThreadId();
 void			Writer_FreeBlockListToPool(FWriteBuffer*, FWriteBuffer*);
diff --git a/Engine/Source/Runtime/TraceLog/Private/Trace/Writer.cpp b/Engine/Source/Runtime/TraceLog/Private/Trace/Writer.cpp
--- a/Engine/Source/Runtime/TraceLog/Private/Trace/Writer.cpp
+++ b/Engine/Source/Runtime/TraceLog/Private/Trace/Writer.cpp
@@ -32,7 +32,7 @@ namespace Private {
 
 ////////////////////////////////////////////////////////////////////////////////
 int32			Encode(const void*, int32, void*, int32);
-uint32			Writer_SendData(uint32, uint8* __restrict, uint32);
+void			Writer_SendData(uint32, uint8* __restrict, uint32);
 void			Writer_InitializeSharedBuffers();
 void			Writer_ShutdownSharedBuffers();
 void			Writer_UpdateSharedBuffers();
@@ -258,7 +258,7 @@ void Writer_SendDataRaw(const void* Data, uint32 Size)
 }
 
 ////////////////////////////////////////////////////////////////////////////////
-uint32 Writer_SendData(uint32 ThreadId, uint8* __restrict Data, uint32 Size)
+void Writer_SendData(uint32 ThreadId, uint8* __restrict Data, uint32 Size)
 {
 	static_assert(ETransport::Active == ETransport::TidPacket, "Active should be set to what the compiled code uses. It is used to track places that assume transport packet format");
 
@@ -268,7 +268,7 @@ uint32 Writer_SendData(uint32 ThreadId, uint8* __restrict Data, uint32 Size)
 
 	if (!GDataHandle)
 	{
-		return 0;
+		return;
 	}
 
 	// Smaller buffers usually aren't redundant enough to benefit from being
@@ -282,8 +282,7 @@ uint32 Writer_SendData(uint32 ThreadId, uint8* __restrict Data, uint32 Size)
 		Packet->PacketSize = uint16(Size);
 
 		Writer_SendDataImpl(Data, Size);
-
-		return Size;
+		return;
 	}
 
 	// Buffer size is expressed as "A + B" where A is a maximum expected
@@ -297,8 +296,6 @@ uint32 Writer_SendData(uint32 ThreadId, uint8* __restrict Data, uint32 Size)
 	Packet.PacketSize += sizeof(FTidPacketEncoded);
 
 	Writer_SendDataImpl(&Packet, Packet.PacketSize);
-
-	return Packet.PacketSize;
 }
 
 ////////////////////////////////////////////////////////////////////////////////
