Narrowed the size of the end-of-text offset. It had to be machine word size prior the introduction of std::atomic. The worker thread reads "committed" to find out what data it is allowed to see. Committed may never make it to the end of the buffer however if the next thing written is too big. So an offset is used to indicate that a buffer should be retired.

#rnx

diff --git a/Engine/Source/Runtime/TraceLog/Private/Trace/TlsBuffer.cpp b/Engine/Source/Runtime/TraceLog/Private/Trace/TlsBuffer.cpp
--- a/Engine/Source/Runtime/TraceLog/Private/Trace/TlsBuffer.cpp
+++ b/Engine/Source/Runtime/TraceLog/Private/Trace/TlsBuffer.cpp
@@ -68,7 +68,7 @@ static FWriteBuffer* Writer_NextBufferInternal()
 	NextBuffer->Cursor = (uint8*)NextBuffer - NextBuffer->Size;
 	NextBuffer->Committed = NextBuffer->Cursor;
 	NextBuffer->Reaped = NextBuffer->Cursor;
-	NextBuffer->EtxOffset = UPTRINT(0) - sizeof(FWriteBuffer);
+	NextBuffer->EtxOffset = 0 - int32(sizeof(FWriteBuffer));
 	NextBuffer->NextBuffer = nullptr;
 
 	FWriteBuffer* CurrentBuffer = GTlsWriteBuffer;
@@ -101,7 +101,7 @@ static FWriteBuffer* Writer_NextBufferInternal()
 		GTlsWriteBuffer = NextBuffer;
 
 		// Retire current buffer.
-		UPTRINT EtxOffset = UPTRINT((uint8*)(CurrentBuffer) - CurrentBuffer->Cursor);
+		int32 EtxOffset = int32(PTRINT((uint8*)(CurrentBuffer) - CurrentBuffer->Cursor));
 		AtomicStoreRelease(&(CurrentBuffer->EtxOffset), EtxOffset);
 	}
 
@@ -146,7 +146,7 @@ static bool Writer_DrainBuffer(uint32 ThreadId, FWriteBuffer* Buffer)
 	}
 
 	// Is this buffer still in use?
-	int32 EtxOffset = int32(AtomicLoadAcquire(&Buffer->EtxOffset));
+	int32 EtxOffset = AtomicLoadAcquire(&Buffer->EtxOffset);
 	return ((uint8*)Buffer - EtxOffset) > Committed;
 }
 
@@ -250,7 +250,7 @@ void Writer_EndThreadBuffer()
 		return;
 	}
 
-	UPTRINT EtxOffset = UPTRINT((uint8*)GTlsWriteBuffer - GTlsWriteBuffer->Cursor);
+	int32 EtxOffset = int32(PTRINT((uint8*)GTlsWriteBuffer - GTlsWriteBuffer->Cursor));
 	AtomicStoreRelaxed(&(GTlsWriteBuffer->EtxOffset), EtxOffset);
 }
 
diff --git a/Engine/Source/Runtime/TraceLog/Public/Trace/Detail/Writer.inl b/Engine/Source/Runtime/TraceLog/Public/Trace/Detail/Writer.inl
--- a/Engine/Source/Runtime/TraceLog/Public/Trace/Detail/Writer.inl
+++ b/Engine/Source/Runtime/TraceLog/Public/Trace/Detail/Writer.inl
@@ -20,7 +20,7 @@ struct FWriteBuffer
 	uint8* __restrict			Cursor;
 	uint8* __restrict volatile	Committed;
 	uint8* __restrict			Reaped;
-	UPTRINT volatile			EtxOffset;
+	int32 volatile				EtxOffset;
 };
 
 ////////////////////////////////////////////////////////////////////////////////
