Indicate if a packet starts with partial data - data that has overflowed from the previous packet and is otherwise unintelljigiblaaes.

#rnx

diff --git a/Engine/Source/Runtime/TraceLog/Private/Trace/Field.cpp b/Engine/Source/Runtime/TraceLog/Private/Trace/Field.cpp
--- a/Engine/Source/Runtime/TraceLog/Private/Trace/Field.cpp
+++ b/Engine/Source/Runtime/TraceLog/Private/Trace/Field.cpp
@@ -38,6 +38,8 @@ static void Field_WriteAuxData(uint32 Index, int32 Size, CallbackType&& Callback
 			}
 
 			Buffer = Writer_NextBuffer(0);
+			Buffer->Partial = 1;
+
 			bCommit = true;
 		}
 
diff --git a/Engine/Source/Runtime/TraceLog/Private/Trace/TlsBuffer.cpp b/Engine/Source/Runtime/TraceLog/Private/Trace/TlsBuffer.cpp
--- a/Engine/Source/Runtime/TraceLog/Private/Trace/TlsBuffer.cpp
+++ b/Engine/Source/Runtime/TraceLog/Private/Trace/TlsBuffer.cpp
@@ -76,6 +76,7 @@ static FWriteBuffer* Writer_NextBufferInternal()
 	{
 		NextBuffer->ThreadId = uint16(Writer_GetThreadId());
 		NextBuffer->PrevTimestamp = TimeGetTimestamp();
+		NextBuffer->Partial = 0;
 
 		GTlsWriteBuffer = NextBuffer;
 
@@ -97,6 +98,7 @@ static FWriteBuffer* Writer_NextBufferInternal()
 		CurrentBuffer->NextBuffer = NextBuffer;
 		NextBuffer->ThreadId = CurrentBuffer->ThreadId;
 		NextBuffer->PrevTimestamp = CurrentBuffer->PrevTimestamp;
+		NextBuffer->Partial = 0;
 
 		GTlsWriteBuffer = NextBuffer;
 
diff --git a/Engine/Source/Runtime/TraceLog/Public/Trace/Detail/Writer.inl b/Engine/Source/Runtime/TraceLog/Public/Trace/Detail/Writer.inl
--- a/Engine/Source/Runtime/TraceLog/Public/Trace/Detail/Writer.inl
+++ b/Engine/Source/Runtime/TraceLog/Public/Trace/Detail/Writer.inl
@@ -21,6 +21,7 @@ struct FWriteBuffer
 	uint8* __restrict volatile	Committed;
 	uint8* __restrict			Reaped;
 	int32 volatile				EtxOffset;
+	int32						Partial;
 };
 
 ////////////////////////////////////////////////////////////////////////////////
