Replaced some magic constants with a something more semantic.

#rnx

diff --git a/Engine/Source/Developer/TraceAnalysis/Private/Analysis/Transport/TidPacketTransport.cpp b/Engine/Source/Developer/TraceAnalysis/Private/Analysis/Transport/TidPacketTransport.cpp
--- a/Engine/Source/Developer/TraceAnalysis/Private/Analysis/Transport/TidPacketTransport.cpp
+++ b/Engine/Source/Developer/TraceAnalysis/Private/Analysis/Transport/TidPacketTransport.cpp
@@ -38,7 +38,7 @@ bool FTidPacketTransport::ReadPacket()
 
 	FTransport::Advance(PacketBase->PacketSize);
 
-	uint32 ThreadId = PacketBase->ThreadId & ~0x8000;
+	uint32 ThreadId = PacketBase->ThreadId & FTidPacketBase::ThreadIdMask;
 	FThreadStream& Thread = FindOrAddThread(ThreadId);
 
 	uint32 DataSize = PacketBase->PacketSize - sizeof(FTidPacketBase);
diff --git a/Engine/Source/Runtime/TraceLog/Private/Trace/Writer.cpp b/Engine/Source/Runtime/TraceLog/Private/Trace/Writer.cpp
--- a/Engine/Source/Runtime/TraceLog/Private/Trace/Writer.cpp
+++ b/Engine/Source/Runtime/TraceLog/Private/Trace/Writer.cpp
@@ -278,7 +278,7 @@ void Writer_SendData(uint32 ThreadId, uint8* __restrict Data, uint32 Size)
 		Data -= sizeof(FTidPacket);
 		Size += sizeof(FTidPacket);
 		auto* Packet = (FTidPacket*)Data;
-		Packet->ThreadId = uint16(ThreadId & 0x7fff);
+		Packet->ThreadId = uint16(ThreadId & FTidPacketBase::ThreadIdMask);
 		Packet->PacketSize = uint16(Size);
 
 		Writer_SendDataImpl(Data, Size);
@@ -290,7 +290,8 @@ void Writer_SendData(uint32 ThreadId, uint8* __restrict Data, uint32 Size)
 	// per LZ4_COMPRESSBOUND.
 	TTidPacketEncoded<8192 + 64> Packet;
 
-	Packet.ThreadId = 0x8000 | uint16(ThreadId & 0x7fff);
+	Packet.ThreadId = FTidPacketBase::EncodedMarker;
+	Packet.ThreadId |= uint16(ThreadId & FTidPacketBase::ThreadIdMask);
 	Packet.DecodedSize = uint16(Size);
 	Packet.PacketSize = Encode(Data, Packet.DecodedSize, Packet.Data, sizeof(Packet.Data));
 	Packet.PacketSize += sizeof(FTidPacketEncoded);
diff --git a/Engine/Source/Runtime/TraceLog/Public/Trace/Detail/Transport.h b/Engine/Source/Runtime/TraceLog/Public/Trace/Detail/Transport.h
--- a/Engine/Source/Runtime/TraceLog/Public/Trace/Detail/Transport.h
+++ b/Engine/Source/Runtime/TraceLog/Public/Trace/Detail/Transport.h
@@ -32,7 +32,11 @@ namespace Private
 ////////////////////////////////////////////////////////////////////////////////
 struct FTidPacketBase
 {
-	enum : uint16 { EncodedMarker = 0x8000 };
+	enum : uint16
+	{
+		EncodedMarker = 0x8000,
+		ThreadIdMask  = EncodedMarker - 1,
+	};
 
 	uint16 PacketSize;
 	uint16 ThreadId;
