File that had escaped from 16218739

#rnx

diff --git a/Engine/Source/Runtime/TraceLog/Private/Trace/Important/Cache.cpp b/Engine/Source/Runtime/TraceLog/Private/Trace/Important/Cache.cpp
--- a/Engine/Source/Runtime/TraceLog/Private/Trace/Important/Cache.cpp
+++ b/Engine/Source/Runtime/TraceLog/Private/Trace/Important/Cache.cpp
@@ -85,7 +85,7 @@ static void Writer_CacheCommit(const FCacheBuffer* Collector)
 	uint32 OutputSize = Encode(Collector->Data, InputSize, Packet->Data, EncodeMaxSize);
 
 	Packet->PacketSize = OutputSize + sizeof(FTidPacketEncoded);
-	Packet->ThreadId = 0x8000 | uint16(ETransportTid::Importants);
+	Packet->ThreadId = FTidPacketBase::EncodedMarker | uint16(ETransportTid::Importants);
 	Packet->DecodedSize = uint16(InputSize);
 
 	Used = sizeof(FTidPacketEncoded) + OutputSize;
