Added a bit to a trace packet's header to indicate if it starts with partial data. 2^14 will be more than sufficient for thread IDs.

#rnx

diff --git a/Engine/Source/Runtime/TraceLog/Public/Trace/Detail/Transport.h b/Engine/Source/Runtime/TraceLog/Public/Trace/Detail/Transport.h
--- a/Engine/Source/Runtime/TraceLog/Public/Trace/Detail/Transport.h
+++ b/Engine/Source/Runtime/TraceLog/Public/Trace/Detail/Transport.h
@@ -35,7 +35,8 @@ struct FTidPacketBase
 	enum : uint16
 	{
 		EncodedMarker = 0x8000,
-		ThreadIdMask  = EncodedMarker - 1,
+		PartialMarker = 0x4000,
+		ThreadIdMask  = PartialMarker - 1,
 	};
 
 	uint16 PacketSize;
