Tail tracing. Maintain a tail of traced events which can be used to give historical context when late-connecting (provided there are trace channels which are enabled by default which there is not yet)

#rnx

diff --git a/Engine/Source/Runtime/TraceLog/Private/Trace/Tail.cpp b/Engine/Source/Runtime/TraceLog/Private/Trace/Tail.cpp
new file mode 100644
--- /dev/null
+++ b/Engine/Source/Runtime/TraceLog/Private/Trace/Tail.cpp
@@ -0,0 +1,372 @@
+// Copyright Epic Games, Inc. All Rights Reserved.
+
+#include "Trace/Config.h"
+
+#if UE_TRACE_ENABLED
+
+#include "Trace/Detail/Protocol.h"
+#include "Trace/Detail/Transport.h"
+
+#include <string.h>
+#include <type_traits>
+#include <initializer_list>
+
+namespace UE {
+namespace Trace {
+namespace Private {
+
+// This here to aid future maintenance of a trace's transport packets.
+static_assert(ETransport::Active == ETransport::TidPacket, "Tail-tracing is transport aware");
+
+////////////////////////////////////////////////////////////////////////////////
+static void	StressRingPacket();
+uint32		GetEncodeMaxSize(uint32);
+int32		Encode(const void*, int32, void*, int32);
+void*		Writer_MemoryAllocate(SIZE_T, uint32);
+void		Writer_MemoryFree(void*, uint32);
+void		Writer_SendData(uint32, uint8* __restrict, uint32);
+void		Writer_SendDataRaw(const void*, uint32);
+
+////////////////////////////////////////////////////////////////////////////////
+// ** See the bottom of this file for an explanation of FPacketRing
+class FPacketRing
+{
+public:
+	struct FRange
+	{
+		const void*	Data;
+		uint32		Size;
+	};
+
+	void			Initialize(uint32 InSize);
+	void			Shutdown();
+	void			Reset();
+	uint32			GetSize() const;
+	bool			IsActive() const;
+	FRange			GetBackPackets() const;
+	FRange			GetFrontPackets() const;
+	template <typename CallbackType>
+	void			IterateRanges(CallbackType&& Callback);
+	template <typename PacketType>
+	PacketType*		Append(uint32 InSize);
+	void			BackUp(uint32 InSize);
+
+private:
+	FTidPacketBase*	AppendImpl(uint32 InSize);
+	uint8*			Data;
+	uint32			Size;
+	uint32			Cursor;
+	uint32			Left;
+	uint32			Right;
+};
+
+// TraceLog must be ready after value-init so that it can be used before
+// dynamic-init. Thus statically-scoped objects cannot have [con|des]tructors.
+static_assert(std::is_trivial<FPacketRing>(), "FPacketRing must be trivial");
+
+////////////////////////////////////////////////////////////////////////////////
+void FPacketRing::Initialize(uint32 InSize)
+{
+	Data = (uint8*)Writer_MemoryAllocate(InSize, 16);
+	Size = InSize;
+	Reset();
+}
+
+////////////////////////////////////////////////////////////////////////////////
+void FPacketRing::Shutdown()
+{
+	Writer_MemoryFree(Data, Size);
+	Data = nullptr;
+}
+
+////////////////////////////////////////////////////////////////////////////////
+void FPacketRing::Reset()
+{
+	Cursor = 0;
+    Left = Right = Size;
+}
+
+////////////////////////////////////////////////////////////////////////////////
+uint32 FPacketRing::GetSize() const
+{
+	return Size;
+}
+
+////////////////////////////////////////////////////////////////////////////////
+bool FPacketRing::IsActive() const
+{
+	return Data != nullptr;
+}
+
+////////////////////////////////////////////////////////////////////////////////
+FPacketRing::FRange FPacketRing::GetBackPackets() const
+{
+	return { Data + Left, Right - Left };
+}
+
+////////////////////////////////////////////////////////////////////////////////
+FPacketRing::FRange FPacketRing::GetFrontPackets() const
+{
+	return { Data, Cursor };
+}
+
+////////////////////////////////////////////////////////////////////////////////
+template <typename CallbackType>
+void FPacketRing::IterateRanges(CallbackType&& Callback)
+{
+	FPacketRing::FRange Ranges[] = { GetBackPackets(), GetFrontPackets() };
+
+	// Send out the ranges.
+	for (const auto& Range : Ranges)
+	{
+		if (Range.Size == 0)
+		{
+			continue;
+		}
+
+		Callback(Range);
+	}
+}
+
+////////////////////////////////////////////////////////////////////////////////
+template <typename PacketType>
+PacketType* FPacketRing::Append(uint32 InSize)
+{
+	FTidPacketBase* Ptr = AppendImpl(InSize + sizeof(PacketType));
+	return static_cast<PacketType*>(Ptr);
+}
+
+////////////////////////////////////////////////////////////////////////////////
+void FPacketRing::BackUp(uint32 InSize)
+{
+	Cursor -= InSize;
+}
+
+////////////////////////////////////////////////////////////////////////////////
+FTidPacketBase* FPacketRing::AppendImpl(uint32 InSize)
+{
+	// ** See the bottom of this file for an explanation of the logic here.
+
+	// Too big to fit in the buffer? It is not possible to maintain consisntency
+	// past this point as some of the data would be truncated. So we drop all
+	// known data and start afresh.
+	if (UNLIKELY(InSize > Size))
+	{
+		Reset();
+		return nullptr;
+	}
+
+	uint32 NextCursor = Cursor + InSize;
+
+	// Run off the end of the buffer?
+	if (UNLIKELY(NextCursor > Size))
+	{
+		Left = 0;
+		Right = Cursor;
+		Cursor = 0;
+		NextCursor = InSize;
+	}
+
+	// Discard old packets until there is space for the new one
+	while (true)
+	{
+		// Does [Cursor, NextCursor) no longer overlap [Left, Right)?
+		if (LIKELY(Left >= NextCursor))
+		{
+			break;
+		}
+
+		// Is [Left, Right) now empty?
+		if (UNLIKELY(Left >= Right))
+		{
+			break;
+		}
+
+		const auto* TidPacket = (const FTidPacketBase*)(Data + Left);
+		Left += TidPacket->PacketSize;
+	}
+
+	// Drop a packet from left.
+	auto* TidPacket = (FTidPacketBase*)(Data + Cursor);
+	TidPacket->PacketSize = InSize;
+
+	Cursor = NextCursor;
+	return TidPacket;
+}
+
+
+
+////////////////////////////////////////////////////////////////////////////////
+static FPacketRing GPacketRing; // = {};
+
+
+
+////////////////////////////////////////////////////////////////////////////////
+void Writer_TailAppend(uint32 ThreadId, uint8* __restrict Data, uint32 Size, bool bPartial)
+{
+	// Perhaps tail tracing is disabled?
+	if (!GPacketRing.IsActive())
+	{
+		return Writer_SendData(ThreadId, Data, Size);
+	}
+
+	// If the packet is too going to be too big (discounting compression ratio as
+	// that's unknown) then 
+	if (uint32(Size + sizeof(FTidPacketEncoded)) > GPacketRing.GetSize())
+	{
+		GPacketRing.Reset();
+		return Writer_SendData(ThreadId, Data, Size);
+	}
+
+	ThreadId &= FTidPacketBase::ThreadIdMask;
+	ThreadId |= bPartial ? FTidPacketBase::PartialMarker : 0;
+
+	// Smaller buffers usually aren't redundant enough to benefit from being
+	// compressed. They often end up being larger.
+	if (Size <= 384)
+	{
+		auto* Packet = GPacketRing.Append<FTidPacket>(Size);
+		Packet->ThreadId = uint16(ThreadId);
+		::memcpy(Packet->Data, Data, Size);
+
+		Writer_SendDataRaw(Packet, Packet->PacketSize);
+		return;
+	}
+
+	uint32 EncodeMaxSize = GetEncodeMaxSize(Size);
+	auto* Packet = GPacketRing.Append<FTidPacketEncoded>(EncodeMaxSize);
+	Packet->ThreadId = uint16(ThreadId);
+	Packet->ThreadId |= FTidPacketBase::EncodedMarker;
+	Packet->DecodedSize = uint16(Size);
+
+	uint32 EncodeSize = Encode(Data, Size, Packet->Data, EncodeMaxSize);
+	uint32 BackUp = EncodeMaxSize - EncodeSize;
+	GPacketRing.BackUp(BackUp);
+	Packet->PacketSize -= uint16(BackUp);
+
+	return Writer_SendDataRaw(Packet, Packet->PacketSize);
+}
+
+////////////////////////////////////////////////////////////////////////////////
+void Writer_TailOnConnect()
+{
+	// If there's no tail being maintained then there is nothing to send.
+	if (!GPacketRing.IsActive())
+	{
+		return;
+	}
+
+	GPacketRing.IterateRanges([] (const FPacketRing::FRange& Range)
+	{
+		Writer_SendDataRaw(Range.Data, Range.Size);
+	});
+}
+
+////////////////////////////////////////////////////////////////////////////////
+void Writer_InitializeTail(int32 BufferSize)
+{
+#if defined(STRESS_PACKET_RING)
+	StressRingPacket();
+#endif
+
+	if (BufferSize <= 0)
+	{
+		return;
+	}
+
+	// Round up to 1K and clamp the size. There has to be a sensible amount of
+	// buffer size for tail tracing to work or be useful.
+	uint32 Rounding = (1 << 10) - 1;
+	BufferSize = (BufferSize + Rounding) & ~Rounding;
+	if (BufferSize < (128 << 10))
+	{
+		BufferSize = 128 << 10;
+	}
+
+	GPacketRing.Initialize(BufferSize);
+}
+
+////////////////////////////////////////////////////////////////////////////////
+void Writer_ShutdownTail()
+{
+	GPacketRing.Shutdown();
+}
+
+
+
+////////////////////////////////////////////////////////////////////////////////
+#if defined(STRESS_PACKET_RING)
+static void StressRingPacket()
+{
+	FPacketRing Ring;
+	Ring.Initialize(300);
+
+	uint32 Bits = 0x0493'0493;
+	for (int32 i = 0; i < 1024; ++i)
+	{
+		FTidPacket* Packet = Ring.Append<FTidPacket>((Bits & 0x1f) + 6);
+		Packet->ThreadId = i;
+		if ((Bits & 0x15) == 0)
+		{
+			Packet->ThreadId |= FTidPacketBase::PartialMarker;
+		}
+
+		Ring.IterateRanges([] (const FPacketRing::FRange&)
+		{
+			/* nop */
+		});
+
+		Bits = (Bits ^ 0xa93a'93a9) * 0x0493;
+	}
+
+	for (int32 i = 7; i < 448; i += 67)
+	{
+		if (auto* Packet = Ring.Append<FTidPacket>(i))
+		{
+			Packet->ThreadId = 0;
+		}
+	}
+}
+#endif // STRESS_PACKET_RING
+
+} // namespace Private
+} // namespace Trace
+} // namespace UE
+
+#endif // UE_TRACE_ENABLED
+
+/*
+
+FPacketRing ring-buffers packets. Internally the buffer is divided up into two
+ranges; [0-Cursor) and [Left-Right) which are initially empty;
+
+ 0                                                                      L
+ C----------------------------------------------------------------------R
+
+A packet consists of a size and a opaque blob of data. Reading the sizes allows
+one to stride through the packets.
+                                                                        L
+ 0[SZ]==============>[SZ]=============>[SZ]=======>C--------------------R
+
+Eventually the next packet will not fit in the buffer because the next cursor (N)
+is off the buffer's end;
+                                                                        L
+ 0[SZ]==============>[SZ]=============>[SZ]=======>[SZ]==========>C-----R
+                                                                  [SZ]========>N
+
+When this happens the 0-Cursor range is transfered to Left-Right and the 0-Cursor
+range is set such that it can contain the new packet being added.
+
+ L[SZ]==============>[SZ]=============>[SZ]=======>[SZ]==========>R-----|
+ 0[SZ]========>C
+
+The two ranges now overlap so packets are then removed from Left until there is
+enough space for the new packet.
+
+ 0[SZ]========>C-----L[SZ]============>[SZ]=======>[SZ]==========>R-----|
+
+The Left-Right range has the oldest packets. Left will eventually advance to meet
+Right at which point the Left-Right range becomes empty The process above repeats
+as if the buffer was being filled for the first time.
+
+*/
diff --git a/Engine/Source/Runtime/TraceLog/Private/Trace/TlsBuffer.cpp b/Engine/Source/Runtime/TraceLog/Private/Trace/TlsBuffer.cpp
--- a/Engine/Source/Runtime/TraceLog/Private/Trace/TlsBuffer.cpp
+++ b/Engine/Source/Runtime/TraceLog/Private/Trace/TlsBuffer.cpp
@@ -14,7 +14,7 @@ namespace Trace {
 namespace Private {
 
 ////////////////////////////////////////////////////////////////////////////////
-void			Writer_SendData(uint32, uint8* __restrict, uint32);
+void			Writer_TailAppend(uint32, uint8* __restrict, uint32, bool);
 FWriteBuffer*	Writer_AllocateBlockFromPool();
 uint32			Writer_GetThreadId();
 void			Writer_FreeBlockListToPool(FWriteBuffer*, FWriteBuffer*);
@@ -143,7 +143,9 @@ static bool Writer_DrainBuffer(uint32 ThreadId, FWriteBuffer* Buffer)
 		BytesReaped += SizeToReap;
 		BytesSent += /*...*/
 #endif
-		Writer_SendData(ThreadId, Buffer->Reaped, SizeToReap);
+		bool bPartial = (Buffer->Partial == 1);
+		bPartial &= UPTRINT(Buffer->Reaped + Buffer->Size) == UPTRINT(Buffer);
+		Writer_TailAppend(ThreadId, Buffer->Reaped, SizeToReap, bPartial);
 		Buffer->Reaped = Committed;
 	}
 
diff --git a/Engine/Source/Runtime/TraceLog/Private/Trace/Writer.cpp b/Engine/Source/Runtime/TraceLog/Private/Trace/Writer.cpp
--- a/Engine/Source/Runtime/TraceLog/Private/Trace/Writer.cpp
+++ b/Engine/Source/Runtime/TraceLog/Private/Trace/Writer.cpp
@@ -33,6 +33,10 @@ namespace Private {
 ////////////////////////////////////////////////////////////////////////////////
 int32			Encode(const void*, int32, void*, int32);
 void			Writer_SendData(uint32, uint8* __restrict, uint32);
+void			Writer_InitializeTail(int32);
+void			Writer_ShutdownTail();
+void			Writer_TailAppend(uint32, uint8* __restrict, uint32, bool=false);
+void			Writer_TailOnConnect();
 void			Writer_InitializeSharedBuffers();
 void			Writer_ShutdownSharedBuffers();
 void			Writer_UpdateSharedBuffers();
@@ -171,7 +175,7 @@ void* Writer_MemoryAllocate(SIZE_T Size, uint32 Alignment)
 	if (TraceData.GetSize())
 	{
 		uint32 ThreadId = Writer_GetThreadId();
-		Writer_SendData(ThreadId, TraceData.GetData(), TraceData.GetSize());
+		Writer_TailAppend(ThreadId, TraceData.GetData(), TraceData.GetSize());
 	}
 
 #if TRACE_PRIVATE_STATISTICS
@@ -216,7 +220,7 @@ void Writer_MemoryFree(void* Address, uint32 Size)
 	if (TraceData.GetSize())
 	{
 		uint32 ThreadId = Writer_GetThreadId();
-		Writer_SendData(ThreadId, TraceData.GetData(), TraceData.GetSize());
+		Writer_TailAppend(ThreadId, TraceData.GetData(), TraceData.GetSize());
 	}
 #endif // TRACE_PRIVATE_STOMP
 
@@ -437,6 +441,7 @@ static bool Writer_UpdateConnection()
 	Writer_SendData(ETransportTid::Internal, HeaderEvents.GetData(), HeaderEvents.GetSize());
 
 	Writer_CacheOnConnect();
+	Writer_TailOnConnect();
 
 	GSerialSyncPending = true;
 	return true;
@@ -578,6 +583,7 @@ static void Writer_InternalShutdown()
 	Writer_ShutdownControl();
 	Writer_ShutdownPool();
 	Writer_ShutdownSharedBuffers();
+	Writer_ShutdownTail();
 
 	GInitialized = false;
 }
@@ -611,6 +617,8 @@ void Writer_InternalInitialize()
 ////////////////////////////////////////////////////////////////////////////////
 void Writer_Initialize(const FInitializeDesc& Desc)
 {
+	Writer_InitializeTail(Desc.TailSizeBytes);
+
 	if (Desc.bUseWorkerThread)
 	{
 		Writer_WorkerCreate();
diff --git a/Engine/Source/Runtime/TraceLog/Public/Trace/Trace.h b/Engine/Source/Runtime/TraceLog/Public/Trace/Trace.h
--- a/Engine/Source/Runtime/TraceLog/Public/Trace/Trace.h
+++ b/Engine/Source/Runtime/TraceLog/Public/Trace/Trace.h
@@ -23,6 +23,7 @@ enum WideString {};
 
 struct FInitializeDesc
 {
+	uint32			TailSizeBytes		= 4 << 20;
 	bool			bUseWorkerThread	= true;
 };
 
