Turn off tail tracing for a bit. A kink has occurred.

#rnx

diff --git a/Engine/Source/Runtime/TraceLog/Private/Trace/Tail.cpp b/Engine/Source/Runtime/TraceLog/Private/Trace/Tail.cpp
--- a/Engine/Source/Runtime/TraceLog/Private/Trace/Tail.cpp
+++ b/Engine/Source/Runtime/TraceLog/Private/Trace/Tail.cpp
@@ -283,7 +283,7 @@ void Writer_InitializeTail(int32 BufferSize)
 		BufferSize = 128 << 10;
 	}
 
-	GPacketRing.Initialize(BufferSize);
+	//GPacketRing.Initialize(BufferSize);
 }
 
 ////////////////////////////////////////////////////////////////////////////////
