Missing macro argument when Trace gets compiled out.

#rnx

diff --git a/Engine/Source/Runtime/TraceLog/Public/Trace/Detail/Trace.h b/Engine/Source/Runtime/TraceLog/Public/Trace/Detail/Trace.h
--- a/Engine/Source/Runtime/TraceLog/Public/Trace/Detail/Trace.h
+++ b/Engine/Source/Runtime/TraceLog/Public/Trace/Detail/Trace.h
@@ -168,7 +168,7 @@ class FChannel;
 #define TRACE_PRIVATE_EVENT_BEGIN(LoggerName, EventName, ...) \
 	TRACE_PRIVATE_EVENT_BEGIN_IMPL(LoggerName, EventName)
 
-#define TRACE_PRIVATE_EVENT_BEGIN_EXTERN(LoggerName, EventName) \
+#define TRACE_PRIVATE_EVENT_BEGIN_EXTERN(LoggerName, EventName, ...) \
 	TRACE_PRIVATE_EVENT_BEGIN_IMPL(LoggerName, EventName)
 
 #define TRACE_PRIVATE_EVENT_BEGIN_IMPL(LoggerName, EventName) \
