Fixed up and incomplete comment

diff --git a/Engine/Source/Runtime/TraceLog/Private/Trace/Tail.cpp b/Engine/Source/Runtime/TraceLog/Private/Trace/Tail.cpp
--- a/Engine/Source/Runtime/TraceLog/Private/Trace/Tail.cpp
+++ b/Engine/Source/Runtime/TraceLog/Private/Trace/Tail.cpp
@@ -210,8 +210,8 @@ void Writer_TailAppend(uint32 ThreadId, uint8* __restrict Data, uint32 Size, boo
 		return Writer_SendData(ThreadId, Data, Size);
 	}
 
-	// If the packet is too going to be too big (discounting compression ratio as
-	// that's unknown) then 
+	// If the packet is going to be too big (discounting compression ratio as
+	// that's unknown) then we'll drop the history and this packet.
 	if (uint32(Size + sizeof(FTidPacketEncoded)) > GPacketRing.GetSize())
 	{
 		GPacketRing.Reset();
