Clear partial flags after a trace packet has been buffered. The packet is passed onto be sent out on any existing connection (or file) and therefore the partiality of the packet is irrelevant.

diff --git a/Engine/Source/Runtime/TraceLog/Private/Trace/Tail.cpp b/Engine/Source/Runtime/TraceLog/Private/Trace/Tail.cpp
--- a/Engine/Source/Runtime/TraceLog/Private/Trace/Tail.cpp
+++ b/Engine/Source/Runtime/TraceLog/Private/Trace/Tail.cpp
@@ -229,6 +229,9 @@ void Writer_TailAppend(uint32 ThreadId, uint8* __restrict Data, uint32 Size, boo
 		Packet->ThreadId = uint16(ThreadId);
 		::memcpy(Packet->Data, Data, Size);
 
+		// If we're already sending data then the partial flag serves no purpose.
+		Packet->ThreadId &= ~FTidPacketBase::PartialMarker;
+
 		Writer_SendDataRaw(Packet, Packet->PacketSize);
 		return;
 	}
@@ -244,6 +247,9 @@ void Writer_TailAppend(uint32 ThreadId, uint8* __restrict Data, uint32 Size, boo
 	GPacketRing.BackUp(BackUp);
 	Packet->PacketSize -= uint16(BackUp);
 
+	// If we're already sending data then the partial flag serves no purpose.
+	Packet->ThreadId &= ~FTidPacketBase::PartialMarker;
+
 	return Writer_SendDataRaw(Packet, Packet->PacketSize);
 }
 
