Removed a spurious return

diff --git a/Engine/Source/Runtime/TraceLog/Private/Trace/Tail.cpp b/Engine/Source/Runtime/TraceLog/Private/Trace/Tail.cpp
--- a/Engine/Source/Runtime/TraceLog/Private/Trace/Tail.cpp
+++ b/Engine/Source/Runtime/TraceLog/Private/Trace/Tail.cpp
@@ -250,7 +250,7 @@ void Writer_TailAppend(uint32 ThreadId, uint8* __restrict Data, uint32 Size, boo
 	// If we're already sending data then the partial flag serves no purpose.
 	Packet->ThreadId &= ~FTidPacketBase::PartialMarker;
 
-	return Writer_SendDataRaw(Packet, Packet->PacketSize);
+	Writer_SendDataRaw(Packet, Packet->PacketSize);
 }
 
 ////////////////////////////////////////////////////////////////////////////////
