Removed the artifical delay that was added to support auto-connect. Previously, auto-connect was defered until we knew if we were running the editor or not. But trace would start collecting traced events and withouth a connection, would drop them. This would result in some early events being missing in auto-connect cases. We no longer wait for the editor and connect much earlier.

#rnx

diff --git a/Engine/Source/Runtime/TraceLog/Private/Trace/Writer.cpp b/Engine/Source/Runtime/TraceLog/Private/Trace/Writer.cpp
--- a/Engine/Source/Runtime/TraceLog/Private/Trace/Writer.cpp
+++ b/Engine/Source/Runtime/TraceLog/Private/Trace/Writer.cpp
@@ -479,26 +479,6 @@ static void Writer_WorkerThread()
 {
 	ThreadRegister(TEXT("Trace"), 0, INT_MAX);
 
-	Writer_UpdateControl();
-
-	// At this point we haven't ever collected any trace events. So we'll stall
-	// for just a little bit to give the user a chance to set up sending the trace
-	// somewhere. This way they get all events since boot, otherwise they'll be
-	// unceremoniously dropped.
-	int32 PrologueMs = 2000;
-	do
-	{
-		const uint32 SleepMs = 100;
-		ThreadSleep(SleepMs);
-		PrologueMs -= SleepMs;
-
-		if (Writer_UpdateConnection())
-		{
-			break;
-		}
-	}
-	while (PrologueMs > 0);
-
 	while (!GWorkerThreadQuit)
 	{
 		Writer_WorkerUpdate();
