Removed UE_TRACE_CHANNEL_MODULE_EXTERN() as it duplicates UE_TRACE_CHANNEL_EXTERN()'s functionality.

diff --git a/Engine/Source/Runtime/TraceLog/Public/Trace/Detail/Trace.h b/Engine/Source/Runtime/TraceLog/Public/Trace/Detail/Trace.h
--- a/Engine/Source/Runtime/TraceLog/Public/Trace/Detail/Trace.h
+++ b/Engine/Source/Runtime/TraceLog/Public/Trace/Detail/Trace.h
@@ -38,9 +38,6 @@ class FChannel;
 	TRACE_PRIVATE_CHANNEL_DECLARE(static, ChannelName) \
 	TRACE_PRIVATE_CHANNEL_IMPL(ChannelName, ##__VA_ARGS__)
 
-#define TRACE_PRIVATE_CHANNEL_MODULE_EXTERN(ModuleApi, ChannelName) \
-	TRACE_PRIVATE_CHANNEL_DECLARE(ModuleApi extern, ChannelName)
-
 #define TRACE_PRIVATE_CHANNEL_DEFINE(ChannelName, ...) \
 	TRACE_PRIVATE_CHANNEL_DECLARE(, ChannelName) \
 	TRACE_PRIVATE_CHANNEL_IMPL(ChannelName, ##__VA_ARGS__)
@@ -156,8 +153,6 @@ class FChannel;
 
 #define TRACE_PRIVATE_CHANNEL_EXTERN(ChannelName, ...)
 
-#define TRACE_PRIVATE_CHANNEL_MODULE_EXTERN(ModuleApi, ChannelName)
-
 #define TRACE_PRIVATE_CHANNEL_DEFINE(ChannelName, ...)
 
 #define TRACE_PRIVATE_CHANNELEXPR_IS_ENABLED(ChannelsExpr) \
diff --git a/Engine/Source/Runtime/TraceLog/Public/Trace/Trace.h b/Engine/Source/Runtime/TraceLog/Public/Trace/Trace.h
--- a/Engine/Source/Runtime/TraceLog/Public/Trace/Trace.h
+++ b/Engine/Source/Runtime/TraceLog/Public/Trace/Trace.h
@@ -70,6 +70,5 @@ UE_TRACE_API void	ThreadGroupEnd() UE_TRACE_IMPL();
 ////////////////////////////////////////////////////////////////////////////////
 #define UE_TRACE_CHANNEL(ChannelName, ...)				TRACE_PRIVATE_CHANNEL(ChannelName, ##__VA_ARGS__)
 #define UE_TRACE_CHANNEL_EXTERN(ChannelName, ...)		TRACE_PRIVATE_CHANNEL_EXTERN(ChannelName, ##__VA_ARGS__)
-#define UE_TRACE_CHANNEL_MODULE_EXTERN(ModuleApi, ChannelName)	TRACE_PRIVATE_CHANNEL_MODULE_EXTERN(ModuleApi, ChannelName)
 #define UE_TRACE_CHANNEL_DEFINE(ChannelName, ...)		TRACE_PRIVATE_CHANNEL_DEFINE(ChannelName, ##__VA_ARGS__)
 #define UE_TRACE_CHANNELEXPR_IS_ENABLED(ChannelsExpr)	TRACE_PRIVATE_CHANNELEXPR_IS_ENABLED(ChannelsExpr)
