Add a small amount of inertia when closing trace IO. The connection check is the first thing to be handled in the update meaning that any pending events prior to closing aren't sent. We can capture these pending events too if we wait one update before calling Close().

#rnx

diff --git a/Engine/Source/Runtime/TraceLog/Private/Trace/Writer.cpp b/Engine/Source/Runtime/TraceLog/Private/Trace/Writer.cpp
--- a/Engine/Source/Runtime/TraceLog/Private/Trace/Writer.cpp
+++ b/Engine/Source/Runtime/TraceLog/Private/Trace/Writer.cpp
@@ -372,16 +372,25 @@ static bool Writer_UpdateConnection()
 		return false;
 	}
 
-	// Is this a close request?
-	if (GPendingDataHandle == ~0ull)
+	// Is this a close request? So that we capture some of the events around
+	// the closure we will add some inertia before enacting the close.
+	static const uint32 CloseInertia = 2;
+	if (GPendingDataHandle >= (~0ull - CloseInertia))
 	{
-		if (GDataHandle)
+		--GPendingDataHandle;
+
+		if (GPendingDataHandle == (~0ull -CloseInertia))
 		{
-			IoClose(GDataHandle);
+			if (GDataHandle)
+			{
+				IoClose(GDataHandle);
+			}
+
+			GDataHandle = 0;
+			GPendingDataHandle = 0;
 		}
 
-		GDataHandle = 0;
-		GPendingDataHandle = 0;
+		return true;
 	}
 
 	// Reject the pending connection if we've already got a connection
