Be explicit about the type of characters we're expecting to set on a string-type field.

diff --git a/Engine/Source/Runtime/TraceLog/Private/Trace/Field.cpp b/Engine/Source/Runtime/TraceLog/Private/Trace/Field.cpp
--- a/Engine/Source/Runtime/TraceLog/Private/Trace/Field.cpp
+++ b/Engine/Source/Runtime/TraceLog/Private/Trace/Field.cpp
@@ -77,7 +77,7 @@ void Field_WriteAuxData(uint32 Index, const uint8* Data, int32 Size)
 }
 
 ////////////////////////////////////////////////////////////////////////////////
-void Field_WriteStringAnsi(uint32 Index, const TCHAR* String, int32 Length)
+void Field_WriteStringAnsi(uint32 Index, const WIDECHAR* String, int32 Length)
 {
 	int32 Size = Length;
 	Size &= (FAuxHeader::SizeLimit - 1);
@@ -104,7 +104,7 @@ void Field_WriteStringAnsi(uint32 Index, const ANSICHAR* String, int32 Length)
 }
 
 ////////////////////////////////////////////////////////////////////////////////
-void Field_WriteStringWide(uint32 Index, const TCHAR* String, int32 Length)
+void Field_WriteStringWide(uint32 Index, const WIDECHAR* String, int32 Length)
 {
 	int32 Size = Length * sizeof(String[0]);
 	Size &= (FAuxHeader::SizeLimit - 1); // (see above)
diff --git a/Engine/Source/Runtime/TraceLog/Public/Trace/Detail/Field.h b/Engine/Source/Runtime/TraceLog/Public/Trace/Detail/Field.h
--- a/Engine/Source/Runtime/TraceLog/Public/Trace/Detail/Field.h
+++ b/Engine/Source/Runtime/TraceLog/Public/Trace/Detail/Field.h
@@ -25,8 +25,8 @@ namespace Private
 ////////////////////////////////////////////////////////////////////////////////
 UE_TRACE_API void Field_WriteAuxData(uint32, const uint8*, int32);
 UE_TRACE_API void Field_WriteStringAnsi(uint32, const ANSICHAR*, int32);
-UE_TRACE_API void Field_WriteStringAnsi(uint32, const TCHAR*, int32);
-UE_TRACE_API void Field_WriteStringWide(uint32, const TCHAR*, int32);
+UE_TRACE_API void Field_WriteStringAnsi(uint32, const WIDECHAR*, int32);
+UE_TRACE_API void Field_WriteStringWide(uint32, const WIDECHAR*, int32);
 
 } // namespace Private
 
diff --git a/Engine/Source/Runtime/TraceLog/Public/Trace/Detail/Important/ImportantLogScope.inl b/Engine/Source/Runtime/TraceLog/Public/Trace/Detail/Important/ImportantLogScope.inl
--- a/Engine/Source/Runtime/TraceLog/Public/Trace/Detail/Important/ImportantLogScope.inl
+++ b/Engine/Source/Runtime/TraceLog/Public/Trace/Detail/Important/ImportantLogScope.inl
@@ -146,12 +146,12 @@ struct FImportantLogScope::FFieldSet<FieldMeta, AnsiString>
 		Scope->Ptr[Scope->AuxCursor] = 0; // null terminator
 	}
 
-	static void Impl(FImportantLogScope* Scope, const TCHAR* String, int32 Length=-1)
+	static void Impl(FImportantLogScope* Scope, const WIDECHAR* String, int32 Length=-1)
 	{
 		if (Length < 0)
 		{
 			Length = 0;
-			for (const TCHAR* c = String; *c; ++c, ++Length);
+			for (const WIDECHAR* c = String; *c; ++c, ++Length);
 		}
 
 		auto* Header = (FAuxHeader*)(Scope->Ptr + Scope->AuxCursor);
@@ -175,15 +175,15 @@ struct FImportantLogScope::FFieldSet<FieldMeta, AnsiString>
 template <typename FieldMeta>
 struct FImportantLogScope::FFieldSet<FieldMeta, WideString>
 {
-	static void Impl(FImportantLogScope* Scope, const TCHAR* String, int32 Length=-1)
+	static void Impl(FImportantLogScope* Scope, const WIDECHAR* String, int32 Length=-1)
 	{
 		if (Length < 0)
 		{
 			Length = 0;
-			for (const TCHAR* c = String; *c; ++c, ++Length);
+			for (const WIDECHAR* c = String; *c; ++c, ++Length);
 		}
 
-		uint32 Size = Length * sizeof(TCHAR);
+		uint32 Size = Length * sizeof(WIDECHAR);
 
 		auto* Header = (FAuxHeader*)(Scope->Ptr + Scope->AuxCursor);
 		Header->Size = Size << 8;
diff --git a/Engine/Source/Runtime/TraceLog/Public/Trace/Detail/LogScope.inl b/Engine/Source/Runtime/TraceLog/Public/Trace/Detail/LogScope.inl
--- a/Engine/Source/Runtime/TraceLog/Public/Trace/Detail/LogScope.inl
+++ b/Engine/Source/Runtime/TraceLog/Public/Trace/Detail/LogScope.inl
@@ -259,12 +259,12 @@ struct FLogScope::FFieldSet<FieldMeta, AnsiString>
 		Field_WriteStringAnsi(Index, String, Length);
 	}
 
-	static void Impl(FLogScope*, const TCHAR* String, int32 Length=-1)
+	static void Impl(FLogScope*, const WIDECHAR* String, int32 Length=-1)
 	{
 		if (Length < 0)
 		{
 			Length = 0;
-			for (const TCHAR* c = String; *c; ++c, ++Length);
+			for (const WIDECHAR* c = String; *c; ++c, ++Length);
 		}
 
 		static const uint32 Index = FieldMeta::Index & 0x7f;
@@ -276,12 +276,12 @@ struct FLogScope::FFieldSet<FieldMeta, AnsiString>
 template <typename FieldMeta>
 struct FLogScope::FFieldSet<FieldMeta, WideString>
 {
-	static void Impl(FLogScope*, const TCHAR* String, int32 Length=-1)
+	static void Impl(FLogScope*, const WIDECHAR* String, int32 Length=-1)
 	{
 		if (Length < 0)
 		{
 			Length = 0;
-			for (const TCHAR* c = String; *c; ++c, ++Length);
+			for (const WIDECHAR* c = String; *c; ++c, ++Length);
 		}
 
 		static const uint32 Index = FieldMeta::Index & 0x7f;
