Very minor tweak to placement of parentheses.

diff --git a/Engine/Source/Runtime/TraceLog/Public/Trace/Detail/LogScope.inl b/Engine/Source/Runtime/TraceLog/Public/Trace/Detail/LogScope.inl
--- a/Engine/Source/Runtime/TraceLog/Public/Trace/Detail/LogScope.inl
+++ b/Engine/Source/Runtime/TraceLog/Public/Trace/Detail/LogScope.inl
@@ -103,7 +103,7 @@ inline FScopedLogScope::~FScopedLogScope()
 		return;
 	}
 
-	uint8 LeaveUid = uint8(EKnownEventUids::LeaveScope << EKnownEventUids::_UidShift);
+	uint8 LeaveUid = uint8(EKnownEventUids::LeaveScope) << EKnownEventUids::_UidShift;
 
 	FWriteBuffer* Buffer = Writer_GetBuffer();
 	if (UNLIKELY(int32((uint8*)Buffer - Buffer->Cursor)) < int32(sizeof(LeaveUid)))
@@ -143,7 +143,7 @@ inline FScopedStampedLogScope::~FScopedStampedLogScope()
 	}
 
 	Stamp <<= 8;
-	Stamp += uint8(EKnownEventUids::LeaveScope_T << EKnownEventUids::_UidShift);
+	Stamp += uint8(EKnownEventUids::LeaveScope_T) << EKnownEventUids::_UidShift;
 	memcpy((uint64*)(Buffer->Cursor), &Stamp, sizeof(Stamp));
 	Buffer->Cursor += sizeof(Stamp);
 
@@ -171,7 +171,7 @@ FORCENOINLINE FLogScope FLogScope::Enter(uint32 ExtraSize)
 template <class T>
 FORCENOINLINE FLogScope FLogScope::ScopedEnter(uint32 ExtraSize)
 {
-	uint8 EnterUid = uint8(EKnownEventUids::EnterScope << EKnownEventUids::_UidShift);
+	uint8 EnterUid = uint8(EKnownEventUids::EnterScope) << EKnownEventUids::_UidShift;
 
 	FWriteBuffer* Buffer = Writer_GetBuffer();
 	if (UNLIKELY(int32((uint8*)Buffer - Buffer->Cursor)) < int32(sizeof(EnterUid)))
@@ -201,7 +201,7 @@ FORCENOINLINE FLogScope FLogScope::ScopedStampedEnter(uint32 ExtraSize)
 
 	Stamp = Writer_GetTimestamp(Buffer);
 	Stamp <<= 8;
-	Stamp += uint8(EKnownEventUids::EnterScope_T << EKnownEventUids::_UidShift);
+	Stamp += uint8(EKnownEventUids::EnterScope_T) << EKnownEventUids::_UidShift;
 	memcpy((uint64*)(Buffer->Cursor), &Stamp, sizeof(Stamp));
 	Buffer->Cursor += sizeof(Stamp);
 
