Removed trace event attachments. They have long since been superceded by array-type fields.

diff --git a/Engine/Source/Runtime/TraceLog/Public/Trace/Detail/Field.h b/Engine/Source/Runtime/TraceLog/Public/Trace/Detail/Field.h
--- a/Engine/Source/Runtime/TraceLog/Public/Trace/Detail/Field.h
+++ b/Engine/Source/Runtime/TraceLog/Public/Trace/Detail/Field.h
@@ -192,15 +192,6 @@ struct TField<InNumFields, InSize, EventProps>
 	};
 };
 
-////////////////////////////////////////////////////////////////////////////////
-// Access to additional data that can be included along with a logged event.
-enum Attachment {};
-template <int InOffset>
-struct TField<0, InOffset, Attachment>
-{
-	enum { Offset = InOffset };
-};
-
 } // namespace Trace
 } // namespace UE
 
diff --git a/Engine/Source/Runtime/TraceLog/Public/Trace/Detail/Important/ImportantLogScope.inl b/Engine/Source/Runtime/TraceLog/Public/Trace/Detail/Important/ImportantLogScope.inl
--- a/Engine/Source/Runtime/TraceLog/Public/Trace/Detail/Important/ImportantLogScope.inl
+++ b/Engine/Source/Runtime/TraceLog/Public/Trace/Detail/Important/ImportantLogScope.inl
@@ -196,13 +196,6 @@ struct FImportantLogScope::FFieldSet<FieldMeta, WideString>
 	}
 };
 
-////////////////////////////////////////////////////////////////////////////////
-template <typename FieldMeta>
-struct FImportantLogScope::FFieldSet<FieldMeta, Attachment>
-{
-	static void Impl(...); /* not implemented */
-};
-
 } // namespace Private
 } // namespace Trace
 } // namespace UE
diff --git a/Engine/Source/Runtime/TraceLog/Public/Trace/Detail/LogScope.inl b/Engine/Source/Runtime/TraceLog/Public/Trace/Detail/LogScope.inl
--- a/Engine/Source/Runtime/TraceLog/Public/Trace/Detail/LogScope.inl
+++ b/Engine/Source/Runtime/TraceLog/Public/Trace/Detail/LogScope.inl
@@ -289,23 +289,6 @@ struct FLogScope::FFieldSet<FieldMeta, WideString>
 	}
 };
 
-template <typename FieldMeta>
-struct FLogScope::FFieldSet<FieldMeta, Attachment>
-{
-	template <typename LambdaType>
-	static void Impl(FLogScope* Scope, LambdaType&& Lambda)
-	{
-		uint8* Dest = (uint8*)(Scope->Ptr) + FieldMeta::Offset;
-		Lambda(Dest);
-	}
-
-	static void Impl(FLogScope* Scope, const void* Data, uint32 Size)
-	{
-		uint8* Dest = (uint8*)(Scope->Ptr) + FieldMeta::Offset;
-		::memcpy(Dest, Data, Size);
-	}
-};
-
 } // namespace Private
 } // namespace Trace
 } // namespace UE
diff --git a/Engine/Source/Runtime/TraceLog/Public/Trace/Detail/Trace.h b/Engine/Source/Runtime/TraceLog/Public/Trace/Detail/Trace.h
--- a/Engine/Source/Runtime/TraceLog/Public/Trace/Detail/Trace.h
+++ b/Engine/Source/Runtime/TraceLog/Public/Trace/Detail/Trace.h
@@ -104,19 +104,9 @@ class FChannel;
 #define TRACE_PRIVATE_EVENT_END() \
 		UE::Trace::EventProps> EventProps_Meta; \
 		EventProps_Meta const EventProps_Private = {}; \
-		typedef UE::Trace::TField<0, EventProps_Meta::Size, UE::Trace::Attachment> Attachment_Meta; \
-		const Attachment_Meta Attachment_Field = {}; \
-		template <typename... Ts> auto Attachment(Ts... ts) const { \
-			AreAttachmentsAllowed<bIsImportant>(); \
-			LogScopeType::FFieldSet<Attachment_Meta, UE::Trace::Attachment>::Impl((LogScopeType*)this, Forward<Ts>(ts)...); \
-			return true; \
-		} \
 		typedef std::conditional<bIsImportant, UE::Trace::Private::FImportantLogScope, UE::Trace::Private::FLogScope>::type LogScopeType; \
 		explicit operator bool () const { return true; } \
 		enum { EventFlags = PartialEventFlags|(EventProps_Meta::NumAuxFields ? UE::Trace::Private::FEventInfo::Flag_MaybeHasAux : 0), }; \
-		template <bool bImportance> static void AreAttachmentsAllowed() { \
-			static_assert(!bImportance, "Important events cannot have attachments"); \
-		} \
 		static_assert( \
 			!bIsImportant || (EventFlags & UE::Trace::Private::FEventInfo::Flag_NoSync), \
 			"Trace events flagged as Important events must be marked NoSync" \
@@ -183,7 +173,6 @@ class FChannel;
 		const FTraceDisabled& FieldName;
 
 #define TRACE_PRIVATE_EVENT_END() \
-		const FTraceDisabled& Attachment; \
 	};
 
 #define TRACE_PRIVATE_LOG(LoggerName, EventName, ...) \
