Flags indicating that a packet of trace data began with partial data were getting cleared in error causing analysis to try and read undefined trace data.

diff --git a/Engine/Source/Runtime/TraceLog/Private/Trace/Tail.cpp b/Engine/Source/Runtime/TraceLog/Private/Trace/Tail.cpp
--- a/Engine/Source/Runtime/TraceLog/Private/Trace/Tail.cpp
+++ b/Engine/Source/Runtime/TraceLog/Private/Trace/Tail.cpp
@@ -229,9 +229,6 @@ void Writer_TailAppend(uint32 ThreadId, uint8* __restrict Data, uint32 Size, boo
 		Packet->ThreadId = uint16(ThreadId);
 		::memcpy(Packet->Data, Data, Size);
 
-		// If we're already sending data then the partial flag serves no purpose.
-		Packet->ThreadId &= ~FTidPacketBase::PartialMarker;
-
 		Writer_SendDataRaw(Packet, Packet->PacketSize);
 		return;
 	}
@@ -247,9 +244,6 @@ void Writer_TailAppend(uint32 ThreadId, uint8* __restrict Data, uint32 Size, boo
 	GPacketRing.BackUp(BackUp);
 	Packet->PacketSize -= uint16(BackUp);
 
-	// If we're already sending data then the partial flag serves no purpose.
-	Packet->ThreadId &= ~FTidPacketBase::PartialMarker;
-
 	Writer_SendDataRaw(Packet, Packet->PacketSize);
 }
 
