Bug fix for startup OOM in trace logging due to infinite allocation loop.

The padding calculation in the function that allocates buffers for trace logs (Writer_CreateSharedBuffer) was off by 4 bytes.  The function is called by a lock free allocator which loops until an allocation succeds that's large enough to fit the requested size.  Because of the calculation error, allocation sizes between 997 and 1000 bytes would result in a block being allocated that was too small, so it would never succeed, and retry the allocation an infinite number of types.

Other sizes would work fine, because there is padding to a multiple of the block size.  I was personally seeing the error specifically because my command line (part of the information written to the log) when running a package build was a precise length that resulted in a size of 1000 bytes when added to a bunch of other strings written to the log.  Besides fixing the bug, I added some comments explaining what the lock free allocator is doing.

#rnx

diff --git a/Engine/Source/Runtime/TraceLog/Private/Trace/Important/SharedBuffer.cpp b/Engine/Source/Runtime/TraceLog/Private/Trace/Important/SharedBuffer.cpp
--- a/Engine/Source/Runtime/TraceLog/Private/Trace/Important/SharedBuffer.cpp
+++ b/Engine/Source/Runtime/TraceLog/Private/Trace/Important/SharedBuffer.cpp
@@ -24,7 +24,7 @@ static FSharedBuffer	GNullSharedBuffer	= { 0, FSharedBuffer::RefInit };
 FSharedBuffer* volatile GSharedBuffer		= &GNullSharedBuffer;
 static FSharedBuffer*	GTailBuffer;		// = nullptr
 static uint32			GTailPreSent;		// = 0
-static const uint32		GBlockSize			= 1024;//16 << 10;
+static const uint32		GBlockSize			= 1024;		// Block size must be a power of two!
 extern FStatistics		GTraceStatistics;
 
 ////////////////////////////////////////////////////////////////////////////////
@@ -33,10 +33,9 @@ static FSharedBuffer* Writer_CreateSharedBuffer(uint32 SizeHint=0)
 	const uint32 OverheadSize = sizeof(FSharedBuffer) + sizeof(uint32);
 
 	uint32 BlockSize = GBlockSize;
-	if (SizeHint > GBlockSize - OverheadSize)
+	if (SizeHint + OverheadSize > GBlockSize)
 	{
-		BlockSize += SizeHint - GBlockSize;
-		BlockSize += sizeof(FSharedBuffer);
+		BlockSize += SizeHint + OverheadSize - GBlockSize;
 		BlockSize += GBlockSize - 1;
 		BlockSize &= ~(GBlockSize - 1);
 	}
@@ -56,12 +55,14 @@ static FSharedBuffer* Writer_CreateSharedBuffer(uint32 SizeHint=0)
 ////////////////////////////////////////////////////////////////////////////////
 FNextSharedBuffer Writer_NextSharedBuffer(FSharedBuffer* Buffer, int32 RegionStart, int32 NegSizeAndRef)
 {
+	// Lock free allocation of the next buffer
 	FSharedBuffer* NextBuffer;
 	while (true)
 	{
 		bool bBufferOwner = (RegionStart >= 0);
 		if (LIKELY(bBufferOwner))
 		{
+			// Allocate the next buffer ourselves
 			uint32 Size = -NegSizeAndRef >> FSharedBuffer::CursorShift;
 			NextBuffer = Writer_CreateSharedBuffer(Size);
 			Buffer->Next = NextBuffer;
@@ -70,6 +71,7 @@ FNextSharedBuffer Writer_NextSharedBuffer(FSharedBuffer* Buffer, int32 RegionSta
 		}
 		else
 		{
+			// Another thread is already allocating the next buffer, wait for that to complete
 			for (;; PlatformYield())
 			{
 				NextBuffer = AtomicLoadAcquire(&GSharedBuffer);
@@ -82,7 +84,9 @@ FNextSharedBuffer Writer_NextSharedBuffer(FSharedBuffer* Buffer, int32 RegionSta
 
 		AtomicAddRelease(&(Buffer->Cursor), int32(FSharedBuffer::RefBit));
 
-		// Try and allocate some space in the next buffer.
+		// Try and allocate some space in the next buffer.  The next buffer may have insufficient
+		// space if another thread used some of the memory we allocated, or if it was a buffer
+		// allocated by another thread that's too small to fit our allocation size.
 		RegionStart = AtomicAddRelaxed(&(NextBuffer->Cursor), NegSizeAndRef);
 		if (LIKELY(RegionStart + NegSizeAndRef >= 0))
 		{
