Corrected spelling mistake in a comment.

#rnx

diff --git a/Engine/Source/Runtime/TraceLog/Private/Trace/Tail.cpp b/Engine/Source/Runtime/TraceLog/Private/Trace/Tail.cpp
--- a/Engine/Source/Runtime/TraceLog/Private/Trace/Tail.cpp
+++ b/Engine/Source/Runtime/TraceLog/Private/Trace/Tail.cpp
@@ -354,7 +354,7 @@ is off the buffer's end;
  0[SZ]==============>[SZ]=============>[SZ]=======>[SZ]==========>C-----R
                                                                   [SZ]========>N
 
-When this happens the 0-Cursor range is transfered to Left-Right and the 0-Cursor
+When this happens the 0-Cursor range is transferred to Left-Right and the 0-Cursor
 range is set such that it can contain the new packet being added.
 
  L[SZ]==============>[SZ]=============>[SZ]=======>[SZ]==========>R-----|
