Removed some old disabled code

#rnx

diff --git a/Engine/Source/Runtime/TraceLog/Private/Trace/TlsBuffer.cpp b/Engine/Source/Runtime/TraceLog/Private/Trace/TlsBuffer.cpp
--- a/Engine/Source/Runtime/TraceLog/Private/Trace/TlsBuffer.cpp
+++ b/Engine/Source/Runtime/TraceLog/Private/Trace/TlsBuffer.cpp
@@ -27,19 +27,6 @@ UE_TRACE_EVENT_BEGIN($Trace, ThreadTiming, NoSync)
 	UE_TRACE_EVENT_FIELD(uint64, BaseTimestamp)
 UE_TRACE_EVENT_END()
 
-#define TRACE_PRIVATE_PERF 0
-#if TRACE_PRIVATE_PERF
-UE_TRACE_EVENT_BEGIN($Trace, WorkerThread)
-	UE_TRACE_EVENT_FIELD(uint32, Cycles)
-	UE_TRACE_EVENT_FIELD(uint32, BytesReaped)
-	UE_TRACE_EVENT_FIELD(uint32, BytesSent)
-UE_TRACE_EVENT_END()
-
-UE_TRACE_EVENT_BEGIN($Trace, Memory)
-	UE_TRACE_EVENT_FIELD(uint32, AllocSize)
-UE_TRACE_EVENT_END()
-#endif // TRACE_PRIVATE_PERF
-
 
 
 ////////////////////////////////////////////////////////////////////////////////
@@ -139,10 +126,6 @@ static bool Writer_DrainBuffer(uint32 ThreadId, FWriteBuffer* Buffer)
 	// Send as much as we can.
 	if (uint32 SizeToReap = uint32(Committed - Buffer->Reaped))
 	{
-#if TRACE_PRIVATE_PERF
-		BytesReaped += SizeToReap;
-		BytesSent += /*...*/
-#endif
 		bool bPartial = (Buffer->Partial == 1);
 		bPartial &= UPTRINT(Buffer->Reaped + Buffer->Size) == UPTRINT(Buffer);
 		Writer_TailAppend(ThreadId, Buffer->Reaped, SizeToReap, bPartial);
@@ -170,12 +153,6 @@ void Writer_DrainBuffers()
 		}
 	};
 
-#if TRACE_PRIVATE_PERF
-	uint64 StartTsc = TimeGetTimestamp();
-	uint32 BytesReaped = 0;
-	uint32 BytesSent = 0;
-#endif
-
 	// Claim ownership of any new thread buffer lists
 	FWriteBuffer* __restrict NewThreadList = AtomicExchangeAcquire(&GNewThreadList, (FWriteBuffer*)nullptr);
 
@@ -229,16 +206,6 @@ void Writer_DrainBuffers()
 		}
 	}
 
-#if TRACE_PRIVATE_PERF
-	UE_TRACE_LOG($Trace, WorkerThread, TraceLogChannel)
-		<< WorkerThread.Cycles(uint32(TimeGetTimestamp() - StartTsc))
-		<< WorkerThread.BytesReaped(BytesReaped)
-		<< WorkerThread.BytesSent(BytesSent);
-
-	UE_TRACE_LOG($Trace, Memory, TraceLogChannel)
-		<< Memory.AllocSize(GPoolUsage);
-#endif // TRACE_PRIVATE_PERF
-
 	// Put the retirees we found back into the system again.
 	if (RetireList.Head != nullptr)
 	{
