Send buffers' data closer to age than by thread.

Trace's normal-event buffers form a list by thread, and for each thread there is a list of its buffers. Prior to this change, the entire buffer list for a thread was sent out. Now just one buffer is sent, and the threads requeued if there's more buffers to go. This multiplexes the threads more in the stream of traced data. Consecutive events are less likely to be so far apart and tails represent more threads.

#rnx

diff --git a/Engine/Source/Runtime/TraceLog/Private/Trace/TlsBuffer.cpp b/Engine/Source/Runtime/TraceLog/Private/Trace/TlsBuffer.cpp
--- a/Engine/Source/Runtime/TraceLog/Private/Trace/TlsBuffer.cpp
+++ b/Engine/Source/Runtime/TraceLog/Private/Trace/TlsBuffer.cpp
@@ -179,30 +179,53 @@ void Writer_DrainBuffers()
 	// a list of that thread's buffers (where it is writing trace events to).
 	for (FWriteBuffer* __restrict Buffer : { ActiveThreadList, NewThreadList })
 	{
-		// For each thread...
-		for (FWriteBuffer* __restrict NextThread; Buffer != nullptr; Buffer = NextThread)
-		{
-			NextThread = Buffer->NextThread;
-			uint32 ThreadId = Buffer->ThreadId;
+		// We'll peel off one buffer from each thread at a time. This way packets
+		// are somewhat closer together by age in the stream.
 
-			// For each of the thread's buffers...
-			for (FWriteBuffer* __restrict NextBuffer; Buffer != nullptr; Buffer = NextBuffer)
+		for (FWriteBuffer* __restrict RetryThreadList = nullptr;;)
+		{
+			for (FWriteBuffer* __restrict NextThread; Buffer != nullptr; Buffer = NextThread)
 			{
+				NextThread = Buffer->NextThread;
+
+				uint32 ThreadId = Buffer->ThreadId;
 				if (Writer_DrainBuffer(ThreadId, Buffer))
 				{
-					break;
+					// Buffer's still in use. So it goes on the main active list and
+					// we move onto the next thread.
+					if (Buffer != nullptr)
+					{
+						Buffer->NextThread = GActiveThreadList;
+						GActiveThreadList = Buffer;
+					}
+
+					continue;
 				}
 
-				// Retire the buffer
-				NextBuffer = Buffer->NextBuffer;
+				// The buffer's full so we will retire it. As this means there
+				// will be another buffer in the list 
+
+				FWriteBuffer* __restrict NextBuffer = Buffer->NextBuffer;
+
 				RetireList.Insert(Buffer);
+
+				if (NextBuffer != nullptr)
+				{
+					NextBuffer->NextThread = RetryThreadList;
+					RetryThreadList = NextBuffer;
+				}
 			}
 
-			if (Buffer != nullptr)
+			// Eventually the only buffers remaining aren't full and therefore
+			// never added to the retry list. At that point there is no more
+			// data that can be sent out.
+			if (RetryThreadList == nullptr)
 			{
-				Buffer->NextThread = GActiveThreadList;
-				GActiveThreadList = Buffer;
+				break;
 			}
+
+			Buffer = RetryThreadList;
+			RetryThreadList = nullptr;
 		}
 	}
 
