Make sure all aux-data-terminals events are sent.

The aux-data-terminal was relying on the next unrealted event's commit for it to get communicated to Trace's worker thread. If the aux-data-terminal was the last event written to a buffer (because the next write was too big to fit in the remaining space) it would never make it out.

#rnx

diff --git a/Engine/Source/Runtime/TraceLog/Public/Trace/Detail/LogScope.h b/Engine/Source/Runtime/TraceLog/Public/Trace/Detail/LogScope.h
--- a/Engine/Source/Runtime/TraceLog/Public/Trace/Detail/LogScope.h
+++ b/Engine/Source/Runtime/TraceLog/Public/Trace/Detail/LogScope.h
@@ -36,6 +36,7 @@ public:
 
 protected:
 	void					Commit() const;
+	void					Commit(FWriteBuffer* __restrict LatestBuffer) const;
 
 private:
 	template <uint32 Flags>
diff --git a/Engine/Source/Runtime/TraceLog/Public/Trace/Detail/LogScope.inl b/Engine/Source/Runtime/TraceLog/Public/Trace/Detail/LogScope.inl
--- a/Engine/Source/Runtime/TraceLog/Public/Trace/Detail/LogScope.inl
+++ b/Engine/Source/Runtime/TraceLog/Public/Trace/Detail/LogScope.inl
@@ -25,6 +25,17 @@ inline void FLogScope::Commit() const
 	AtomicStoreRelease((uint8**) &(Buffer->Committed), Buffer->Cursor);
 }
 
+////////////////////////////////////////////////////////////////////////////////
+inline void FLogScope::Commit(FWriteBuffer* __restrict LatestBuffer) const
+{
+	if (LatestBuffer != Buffer)
+	{
+		AtomicStoreRelease((uint8**) &(LatestBuffer->Committed), LatestBuffer->Cursor);
+	}
+
+	Commit();
+}
+
 ////////////////////////////////////////////////////////////////////////////////
 template <uint32 Flags>
 inline auto FLogScope::EnterImpl(uint32 Uid, uint32 Size)
@@ -95,7 +106,7 @@ inline void TLogScope<true>::operator += (const FLogScope&) const
 	LatestBuffer->Cursor[0] = uint8(EKnownEventUids::AuxDataTerminal << EKnownEventUids::_UidShift);
 	LatestBuffer->Cursor++;
 
-	Commit();
+	Commit(LatestBuffer);
 }
 
 
