Add sync points into the trace stream.

These will help analysis know where the world was sampled at runtime in relation to the stream of data it is receiving. With these discrete points in time it is then possible to reason about gaps in events' serial numbers that can occur when late-connecting.

#rnx

diff --git a/Engine/Source/Developer/TraceAnalysis/Private/Analysis/Engine.cpp b/Engine/Source/Developer/TraceAnalysis/Private/Analysis/Engine.cpp
--- a/Engine/Source/Developer/TraceAnalysis/Private/Analysis/Engine.cpp
+++ b/Engine/Source/Developer/TraceAnalysis/Private/Analysis/Engine.cpp
@@ -3239,10 +3239,11 @@ FEstablishTransportStage::EStatus FEstablishTransportStage::OnData(
 	FTransport* Transport = nullptr;
 	switch (Header->TransportVersion)
 	{
-	case ETransport::Raw:		Transport = new FTransport(); break;
-	case ETransport::Packet:	Transport = new FPacketTransport(); break;
-	case ETransport::TidPacket:	Transport = new FTidPacketTransport(); break;
-	default:					return EStatus::Error;
+	case ETransport::Raw:			Transport = new FTransport(); break;
+	case ETransport::Packet:		Transport = new FPacketTransport(); break;
+	case ETransport::TidPacket:		Transport = new FTidPacketTransport(); break;
+	case ETransport::TidPacketSync:	Transport = new FTidPacketTransportSync(); break;
+	default:						return EStatus::Error;
 	}
 
 	uint32 ProtocolVersion = Header->ProtocolVersion;
diff --git a/Engine/Source/Developer/TraceAnalysis/Private/Analysis/Transport/TidPacketTransport.cpp b/Engine/Source/Developer/TraceAnalysis/Private/Analysis/Transport/TidPacketTransport.cpp
--- a/Engine/Source/Developer/TraceAnalysis/Private/Analysis/Transport/TidPacketTransport.cpp
+++ b/Engine/Source/Developer/TraceAnalysis/Private/Analysis/Transport/TidPacketTransport.cpp
@@ -39,6 +39,14 @@ bool FTidPacketTransport::ReadPacket()
 	FTransport::Advance(PacketBase->PacketSize);
 
 	uint32 ThreadId = PacketBase->ThreadId & FTidPacketBase::ThreadIdMask;
+
+	if (ThreadId == ETransportTid::Sync)
+	{
+		++Synced;
+		return false;	// Do not read any more packets. Gives consumers a
+						// chance to sample the world at each known sync point.
+	}
+
 	bool bIsPartial = !!(PacketBase->ThreadId & FTidPacketBase::PartialMarker);
 	FThreadStream* Thread = FindOrAddThread(ThreadId, !bIsPartial);
 	if (Thread == nullptr)
diff --git a/Engine/Source/Developer/TraceAnalysis/Private/Analysis/Transport/TidPacketTransport.h b/Engine/Source/Developer/TraceAnalysis/Private/Analysis/Transport/TidPacketTransport.h
--- a/Engine/Source/Developer/TraceAnalysis/Private/Analysis/Transport/TidPacketTransport.h
+++ b/Engine/Source/Developer/TraceAnalysis/Private/Analysis/Transport/TidPacketTransport.h
@@ -23,6 +23,7 @@ public:
 	uint32					GetThreadCount() const;
 	FStreamReader*			GetThreadStream(uint32 Index);
 	uint32					GetThreadId(uint32 Index) const;
+	uint32					GetSyncCount() const;
 
 private:
 	struct FThreadStream
@@ -37,6 +38,28 @@ private:
 								{ {}, ETransportTid::Events },
 								{ {}, ETransportTid::Importants },
 							};
+
+protected:
+	uint32					Synced = 0x7fff'ffff;
+};
+
+////////////////////////////////////////////////////////////////////////////////
+inline uint32 FTidPacketTransport::GetSyncCount() const
+{
+	return Synced;
+}
+
+
+
+////////////////////////////////////////////////////////////////////////////////
+class FTidPacketTransportSync
+	: public FTidPacketTransport
+{
+public:
+	FTidPacketTransportSync()
+	{
+		Synced = 0;
+	}
 };
 
 } // namespace Trace
diff --git a/Engine/Source/Runtime/TraceLog/Private/Trace/Important/Cache.cpp b/Engine/Source/Runtime/TraceLog/Private/Trace/Important/Cache.cpp
--- a/Engine/Source/Runtime/TraceLog/Private/Trace/Important/Cache.cpp
+++ b/Engine/Source/Runtime/TraceLog/Private/Trace/Important/Cache.cpp
@@ -153,7 +153,7 @@ void Writer_InitializeCache()
 	GCacheActiveBuffer = Writer_CacheCreateBuffer(GCacheBufferSize);
 	GCacheActiveBuffer->TailNext = &GCacheHeadBuffer;
 
-	static_assert(ETransport::Active == ETransport::TidPacket, "The important cache is transport aware");
+	static_assert(ETransport::Active == ETransport::TidPacketSync, "The important cache is transport aware");
 }
 
 ////////////////////////////////////////////////////////////////////////////////
diff --git a/Engine/Source/Runtime/TraceLog/Private/Trace/Tail.cpp b/Engine/Source/Runtime/TraceLog/Private/Trace/Tail.cpp
--- a/Engine/Source/Runtime/TraceLog/Private/Trace/Tail.cpp
+++ b/Engine/Source/Runtime/TraceLog/Private/Trace/Tail.cpp
@@ -16,7 +16,7 @@ namespace Trace {
 namespace Private {
 
 // This here to aid future maintenance of a trace's transport packets.
-static_assert(ETransport::Active == ETransport::TidPacket, "Tail-tracing is transport aware");
+static_assert(ETransport::Active == ETransport::TidPacketSync, "Tail-tracing is transport aware");
 
 ////////////////////////////////////////////////////////////////////////////////
 static void	StressRingPacket();
diff --git a/Engine/Source/Runtime/TraceLog/Private/Trace/Writer.cpp b/Engine/Source/Runtime/TraceLog/Private/Trace/Writer.cpp
--- a/Engine/Source/Runtime/TraceLog/Private/Trace/Writer.cpp
+++ b/Engine/Source/Runtime/TraceLog/Private/Trace/Writer.cpp
@@ -257,7 +257,7 @@ void Writer_SendDataRaw(const void* Data, uint32 Size)
 ////////////////////////////////////////////////////////////////////////////////
 void Writer_SendData(uint32 ThreadId, uint8* __restrict Data, uint32 Size)
 {
-	static_assert(ETransport::Active == ETransport::TidPacket, "Active should be set to what the compiled code uses. It is used to track places that assume transport packet format");
+	static_assert(ETransport::Active == ETransport::TidPacketSync, "Active should be set to what the compiled code uses. It is used to track places that assume transport packet format");
 
 #if TRACE_PRIVATE_STATISTICS
 	GTraceStatistics.BytesTraced += Size;
@@ -345,8 +345,30 @@ static void Writer_DescribeAnnounce()
 
 
 ////////////////////////////////////////////////////////////////////////////////
-static void Writer_LogHeader()
+static int8			GSyncPacketCountdown;	// = 0
+static const int8	GNumSyncPackets			= 3;
+
+////////////////////////////////////////////////////////////////////////////////
+static void Writer_SendSync()
 {
+	if (GSyncPacketCountdown <= 0)
+	{
+		return;
+	}
+
+	// It is possible that some events get collected and discarded by a previous
+	// update that are newer than events sent it the following update where IO
+	// is established. This will result in holes in serial numbering. A few sync
+	// points are sent to aid analysis in determining what are holes and what is
+	// just a requirement for more data. Holws will only occurr at the start.
+
+	// Note that Sync is alias as Important/Internal as changing Bias would
+	// break backwards compatibility.
+
+	FTidPacketBase SyncPacket = { sizeof(SyncPacket), ETransportTid::Sync };
+	Writer_SendDataImpl(&SyncPacket, sizeof(SyncPacket));
+
+	--GSyncPacketCountdown;
 }
 
 ////////////////////////////////////////////////////////////////////////////////
@@ -407,7 +429,7 @@ static bool Writer_UpdateConnection()
 
 	// Stream header
 	const struct {
-		uint8 TransportVersion	= ETransport::TidPacket;
+		uint8 TransportVersion	= ETransport::TidPacketSync;
 		uint8 ProtocolVersion	= EProtocol::Id;
 	} TransportHeader;
 	bOk &= IoWrite(GDataHandle, &TransportHeader, sizeof(TransportHeader));
@@ -431,6 +453,9 @@ static bool Writer_UpdateConnection()
 	Writer_CacheOnConnect();
 	Writer_TailOnConnect();
 
+	// See Writer_SendSync() for details.
+	GSyncPacketCountdown = GNumSyncPackets;
+
 	return true;
 }
 
@@ -448,6 +473,7 @@ static void Writer_WorkerUpdate()
 	Writer_DescribeAnnounce();
 	Writer_UpdateSharedBuffers();
 	Writer_DrainBuffers();
+	Writer_SendSync();
 }
 
 ////////////////////////////////////////////////////////////////////////////////
diff --git a/Engine/Source/Runtime/TraceLog/Public/Trace/Detail/Transport.h b/Engine/Source/Runtime/TraceLog/Public/Trace/Detail/Transport.h
--- a/Engine/Source/Runtime/TraceLog/Public/Trace/Detail/Transport.h
+++ b/Engine/Source/Runtime/TraceLog/Public/Trace/Detail/Transport.h
@@ -8,11 +8,12 @@ namespace Trace {
 ////////////////////////////////////////////////////////////////////////////////
 enum ETransport : uint8
 {
-	_Unused		= 0,
-	Raw			= 1,
-	Packet		= 2,
-	TidPacket	= 3,
-	Active		= TidPacket,
+	_Unused			= 0,
+	Raw				= 1,
+	Packet			= 2,
+	TidPacket		= 3,
+	TidPacketSync	= 4,
+	Active			= TidPacketSync,
 };
 
 ////////////////////////////////////////////////////////////////////////////////
@@ -21,7 +22,10 @@ enum ETransportTid : uint32
 	Events		= 0,			// used to describe events
 	Internal	= 1,			// events to make the trace stream function
 	Importants	= Internal,		// important/cached events
-	Bias,
+	Bias,						// [Bias,End] = threads. Note bias can't be..
+	/* ... */					// ..changed as it breaks backwards compat :(
+	End			= 0x3ffe,		// two msbs are user for packet markers
+	Sync		= 0x3fff,		// see Writer_SendSync()
 };
 
 
