Allow LZ4 to be provided externally.

#rnx

diff --git a/Engine/Source/Runtime/TraceLog/Private/Trace/Codec.cpp b/Engine/Source/Runtime/TraceLog/Private/Trace/Codec.cpp
--- a/Engine/Source/Runtime/TraceLog/Private/Trace/Codec.cpp
+++ b/Engine/Source/Runtime/TraceLog/Private/Trace/Codec.cpp
@@ -8,9 +8,14 @@ THIRD_PARTY_INCLUDES_START
 #	pragma warning(disable : 6239)
 #endif
 
-#define LZ4_NAMESPACE Trace
-#include "LZ4/lz4.c.inl"
-#undef LZ4_NAMESPACE
+#if !defined(TRACE_PRIVATE_EXTERNAL_LZ4)
+#	define LZ4_NAMESPACE Trace
+#		include "LZ4/lz4.c.inl"
+#	undef LZ4_NAMESPACE
+#	define TRACE_PRIVATE_LZ4_NAMESPACE ::Trace::
+#else
+#	define TRACE_PRIVATE_LZ4_NAMESPACE
+#endif
 
 #if defined(_MSC_VER)
 #	pragma warning(pop)
@@ -24,7 +29,7 @@ namespace Private {
 ////////////////////////////////////////////////////////////////////////////////
 int32 Encode(const void* Src, int32 SrcSize, void* Dest, int32 DestSize)
 {
-	return ::Trace::LZ4_compress_fast(
+	return TRACE_PRIVATE_LZ4_NAMESPACE LZ4_compress_fast(
 		(const char*)Src,
 		(char*)Dest,
 		SrcSize,
@@ -42,7 +47,7 @@ uint32 GetEncodeMaxSize(uint32 InputSize)
 ////////////////////////////////////////////////////////////////////////////////
 TRACELOG_API int32 Decode(const void* Src, int32 SrcSize, void* Dest, int32 DestSize)
 {
-	return ::Trace::LZ4_decompress_safe((const char*)Src, (char*)Dest, SrcSize, DestSize);
+	return TRACE_PRIVATE_LZ4_NAMESPACE LZ4_decompress_safe((const char*)Src, (char*)Dest, SrcSize, DestSize);
 }
 
 } // namespace Private
