Use FCStringAnsi instead of FPlatformString so it is easier to support amalgamation.

#rnx

diff --git a/Engine/Source/Runtime/TraceLog/Private/Trace/Control.cpp b/Engine/Source/Runtime/TraceLog/Private/Trace/Control.cpp
--- a/Engine/Source/Runtime/TraceLog/Private/Trace/Control.cpp
+++ b/Engine/Source/Runtime/TraceLog/Private/Trace/Control.cpp
@@ -324,12 +324,12 @@ void Writer_InitializeControl()
 			ANSICHAR Channels[BufferSize] = {};
 			ANSICHAR* Ctx;
 			const bool bState = (ArgV[1][0] != '0');
-			FPlatformString::Strcpy(Channels, BufferSize, ArgV[0]);
-			ANSICHAR* Channel = FPlatformString::Strtok(Channels, ",", &Ctx);
+			FCStringAnsi::Strcpy(Channels, BufferSize, ArgV[0]);
+			ANSICHAR* Channel = FCStringAnsi::Strtok(Channels, ",", &Ctx);
 			while (Channel)
 			{
 				FChannel::Toggle(Channel, bState);
-				Channel = FPlatformString::Strtok(nullptr, ",", &Ctx);
+				Channel = FCStringAnsi::Strtok(nullptr, ",", &Ctx);
 			}
 		}
 	);
