Added PLATFORM_X checks to better exclude those that aren't relevant.

#rnx

diff --git a/Engine/Source/Runtime/TraceLog/Private/Trace/Detail/Android/AndroidTrace.cpp b/Engine/Source/Runtime/TraceLog/Private/Trace/Detail/Android/AndroidTrace.cpp
--- a/Engine/Source/Runtime/TraceLog/Private/Trace/Detail/Android/AndroidTrace.cpp
+++ b/Engine/Source/Runtime/TraceLog/Private/Trace/Detail/Android/AndroidTrace.cpp
@@ -2,7 +2,7 @@
 
 #include "Trace/Config.h"
 
-#if UE_TRACE_ENABLED
+#if UE_TRACE_ENABLED && PLATFORM_ANDROID
 
 #include <arpa/inet.h>
 #include <errno.h>
diff --git a/Engine/Source/Runtime/TraceLog/Private/Trace/Detail/Apple/AppleTrace.cpp b/Engine/Source/Runtime/TraceLog/Private/Trace/Detail/Apple/AppleTrace.cpp
--- a/Engine/Source/Runtime/TraceLog/Private/Trace/Detail/Apple/AppleTrace.cpp
+++ b/Engine/Source/Runtime/TraceLog/Private/Trace/Detail/Apple/AppleTrace.cpp
@@ -2,7 +2,7 @@
 
 #include "Trace/Config.h"
 
-#if UE_TRACE_ENABLED
+#if UE_TRACE_ENABLED && PLATFORM_APPLE
 
 #include <arpa/inet.h>
 #include <errno.h>
diff --git a/Engine/Source/Runtime/TraceLog/Private/Trace/Detail/HoloLens/HoloLensTrace.cpp b/Engine/Source/Runtime/TraceLog/Private/Trace/Detail/HoloLens/HoloLensTrace.cpp
--- a/Engine/Source/Runtime/TraceLog/Private/Trace/Detail/HoloLens/HoloLensTrace.cpp
+++ b/Engine/Source/Runtime/TraceLog/Private/Trace/Detail/HoloLens/HoloLensTrace.cpp
@@ -2,7 +2,7 @@
 
 #include "Trace/Config.h"
 
-#if UE_TRACE_ENABLED
+#if UE_TRACE_ENABLED && PLATFORM_HOLOLENS
 
 #include "Windows/AllowWindowsPlatformTypes.h"
 #	include "Windows/WindowsHWrapper.h"
diff --git a/Engine/Source/Runtime/TraceLog/Private/Trace/Detail/Unix/UnixTrace.cpp b/Engine/Source/Runtime/TraceLog/Private/Trace/Detail/Unix/UnixTrace.cpp
--- a/Engine/Source/Runtime/TraceLog/Private/Trace/Detail/Unix/UnixTrace.cpp
+++ b/Engine/Source/Runtime/TraceLog/Private/Trace/Detail/Unix/UnixTrace.cpp
@@ -2,7 +2,7 @@
 
 #include "Trace/Config.h"
 
-#if UE_TRACE_ENABLED
+#if UE_TRACE_ENABLED && PLATFORM_UNIX
 
 #include <errno.h>
 #include <fcntl.h>
diff --git a/Engine/Source/Runtime/TraceLog/Private/Trace/Detail/Windows/WindowsTrace.cpp b/Engine/Source/Runtime/TraceLog/Private/Trace/Detail/Windows/WindowsTrace.cpp
--- a/Engine/Source/Runtime/TraceLog/Private/Trace/Detail/Windows/WindowsTrace.cpp
+++ b/Engine/Source/Runtime/TraceLog/Private/Trace/Detail/Windows/WindowsTrace.cpp
@@ -2,7 +2,7 @@
 
 #include "Trace/Config.h"
 
-#if UE_TRACE_ENABLED
+#if UE_TRACE_ENABLED && PLATFORM_WINDOWS
 
 #include "Windows/AllowWindowsPlatformTypes.h"
 #	include "Windows/WindowsHWrapper.h"
