Moved protocol define to Config.h so that the block of protocol.h includes are not dependent on being nested

#rnx

diff --git a/Engine/Source/Runtime/TraceLog/Public/Trace/Config.h b/Engine/Source/Runtime/TraceLog/Public/Trace/Config.h
--- a/Engine/Source/Runtime/TraceLog/Public/Trace/Config.h
+++ b/Engine/Source/Runtime/TraceLog/Public/Trace/Config.h
@@ -15,3 +15,7 @@
 #if !defined(UE_TRACE_ENABLED)
 #	define UE_TRACE_ENABLED 0
 #endif
+
+#if UE_TRACE_ENABLED
+#	define TRACE_PRIVATE_PROTOCOL_5
+#endif
diff --git a/Engine/Source/Runtime/TraceLog/Public/Trace/Detail/Protocol.h b/Engine/Source/Runtime/TraceLog/Public/Trace/Detail/Protocol.h
--- a/Engine/Source/Runtime/TraceLog/Public/Trace/Detail/Protocol.h
+++ b/Engine/Source/Runtime/TraceLog/Public/Trace/Detail/Protocol.h
@@ -5,10 +5,6 @@
 #include "Trace/Config.h"
 
 ////////////////////////////////////////////////////////////////////////////////
-#if UE_TRACE_ENABLED
-#	define TRACE_PRIVATE_PROTOCOL_5
-#endif
-
 #if defined(_MSC_VER)
 	#pragma warning(push)
 	#pragma warning(disable : 4200) // non-standard zero-sized array
