Bumped important-event cache page size so there's less waste at the end.

LZ4 compression needs one contiguous piece of memory of worst-case size. This can result in not insignificant waste at the end of each page of cached important event data. By bumping the size the wastage happens less often. Perhaps a better way would be to check if LZ4 can compress in two steps.
#rnx

diff --git a/Engine/Source/Runtime/TraceLog/Private/Trace/Important/Cache.cpp b/Engine/Source/Runtime/TraceLog/Private/Trace/Important/Cache.cpp
--- a/Engine/Source/Runtime/TraceLog/Private/Trace/Important/Cache.cpp
+++ b/Engine/Source/Runtime/TraceLog/Private/Trace/Important/Cache.cpp
@@ -40,7 +40,7 @@ struct alignas(16) FCacheBuffer
 
 
 ////////////////////////////////////////////////////////////////////////////////
-static const uint32		GCacheBufferSize	= 4 << 10;
+static const uint32		GCacheBufferSize	= 64 << 10;
 static const uint32		GCacheCollectorSize	= 1 << 10;
 static FCacheBuffer*	GCacheCollector;	// = nullptr;
 static FCacheBuffer*	GCacheActiveBuffer;	// = nullptr;
