Simple conversion function can handle all types of input characters.

#rnx

diff --git a/Engine/Source/Runtime/TraceLog/Private/Trace/Trace.cpp b/Engine/Source/Runtime/TraceLog/Private/Trace/Trace.cpp
--- a/Engine/Source/Runtime/TraceLog/Private/Trace/Trace.cpp
+++ b/Engine/Source/Runtime/TraceLog/Private/Trace/Trace.cpp
@@ -32,10 +32,10 @@ extern FStatistics GTraceStatistics;
 
 
 ////////////////////////////////////////////////////////////////////////////////
-template <int DestSize>
-static uint32 ToAnsiCheap(ANSICHAR (&Dest)[DestSize], const WIDECHAR* Src)
+template <int DestSize, typename SRC_TYPE>
+static uint32 ToAnsiCheap(ANSICHAR (&Dest)[DestSize], const SRC_TYPE* Src)
 {
-	const WIDECHAR* Cursor = Src;
+	const SRC_TYPE* Cursor = Src;
 	for (ANSICHAR& Out : Dest)
 	{
 		Out = ANSICHAR(*Cursor++ & 0x7f);
