Use internal method instead of Strlen() followed by the same style of loop in Field_WriteString(). Removes the use of FCString making amalgamation for external use easier.

#rnx

diff --git a/Engine/Source/Runtime/TraceLog/Private/Trace/Trace.cpp b/Engine/Source/Runtime/TraceLog/Private/Trace/Trace.cpp
--- a/Engine/Source/Runtime/TraceLog/Private/Trace/Trace.cpp
+++ b/Engine/Source/Runtime/TraceLog/Private/Trace/Trace.cpp
@@ -145,19 +145,23 @@ UE_TRACE_EVENT_END()
 ////////////////////////////////////////////////////////////////////////////////
 void ThreadRegister(const TCHAR* Name, uint32 SystemId, int32 SortHint)
 {
+	ANSICHAR NameA[96];
+
 	uint32 ThreadId = Private::Writer_GetThreadId();
-	uint32 NameLen = FCString::Strlen(Name);
+	uint32 NameLen = ToAnsiCheap(NameA, Name);
 	UE_TRACE_LOG($Trace, ThreadInfo, TraceLogChannel, NameLen * sizeof(ANSICHAR))
 		<< ThreadInfo.ThreadId(ThreadId)
 		<< ThreadInfo.SystemId(SystemId)
 		<< ThreadInfo.SortHint(SortHint)
-		<< ThreadInfo.Name(Name, NameLen);
+		<< ThreadInfo.Name(NameA, NameLen);
 }
 
 ////////////////////////////////////////////////////////////////////////////////
 void ThreadGroupBegin(const TCHAR* Name)
 {
-	uint32 NameLen = FCString::Strlen(Name);
+	ANSICHAR NameA[96];
+
+	uint32 NameLen = ToAnsiCheap(NameA, Name);
 	UE_TRACE_LOG($Trace, ThreadGroupBegin, TraceLogChannel, NameLen * sizeof(ANSICHAR))
 		<< ThreadGroupBegin.Name(Name, NameLen);
 }
