Trace's command and control can be conditionally compiled out.

#rnx

diff --git a/Engine/Source/Runtime/TraceLog/Private/Trace/Control.cpp b/Engine/Source/Runtime/TraceLog/Private/Trace/Control.cpp
--- a/Engine/Source/Runtime/TraceLog/Private/Trace/Control.cpp
+++ b/Engine/Source/Runtime/TraceLog/Private/Trace/Control.cpp
@@ -15,6 +15,8 @@ namespace UE {
 namespace Trace {
 namespace Private {
 
+#if !defined(TRACE_PRIVATE_CONTROL_ENABLED) || TRACE_PRIVATE_CONTROL_ENABLED
+
 ////////////////////////////////////////////////////////////////////////////////
 bool	Writer_SendTo(const ANSICHAR*, uint32=0);
 bool	Writer_WriteTo(const ANSICHAR*);
@@ -345,6 +347,15 @@ void Writer_ShutdownControl()
 	}
 }
 
+#else
+
+void	Writer_InitializeControl()	{}
+void	Writer_ShutdownControl()	{}
+void	Writer_UpdateControl()		{}
+uint32	Writer_GetControlPort()		{ return ~0u; }
+
+#endif // TRACE_PRIVATE_CONTROL_ENABLED
+
 } // namespace Private
 } // namespace Trace
 } // namespace UE
