Corrected statistics collection

Cache wasted was calculated incorrectly. Bytes-sent was missing some bytes. Memory was bumped to 64 bits just in case a 4GB accident happens.

#rnx

diff --git a/Engine/Source/Runtime/TraceLog/Private/Trace/Important/Cache.cpp b/Engine/Source/Runtime/TraceLog/Private/Trace/Important/Cache.cpp
--- a/Engine/Source/Runtime/TraceLog/Private/Trace/Important/Cache.cpp
+++ b/Engine/Source/Runtime/TraceLog/Private/Trace/Important/Cache.cpp
@@ -66,6 +66,10 @@ static void Writer_CacheCommit(const FCacheBuffer* Collector)
 	uint32 EncodeMaxSize = GetEncodeMaxSize(InputSize);
 	if (EncodeMaxSize + sizeof(FTidPacketEncoded) > GCacheActiveBuffer->Remaining)
 	{
+#if TRACE_PRIVATE_STATISTICS
+		GTraceStatistics.CacheWaste += GCacheActiveBuffer->Remaining;
+#endif
+
 		// Retire active buffer
 		*(GCacheActiveBuffer->TailNext) = GCacheActiveBuffer;
 		GCacheActiveBuffer->TailNext = nullptr;
@@ -74,10 +78,6 @@ static void Writer_CacheCommit(const FCacheBuffer* Collector)
 		FCacheBuffer* NewBuffer = Writer_CacheCreateBuffer(GCacheBufferSize);
 		NewBuffer->TailNext = &(GCacheActiveBuffer->Next);
 		GCacheActiveBuffer = NewBuffer;
-
-#if TRACE_PRIVATE_STATISTICS
-		GTraceStatistics.CacheWaste += GCacheActiveBuffer->Remaining;
-#endif
 	}
 
 	uint32 Used = GCacheActiveBuffer->Size - GCacheActiveBuffer->Remaining;
@@ -101,6 +101,11 @@ void Writer_CacheData(uint8* Data, uint32 Size)
 {
 	Writer_SendData(ETransportTid::Importants, Data, Size);
 
+	if (GCacheCollector == nullptr)
+	{
+		return;
+	}
+
 	while (true)
 	{
 		uint32 StepSize = (Size < GCacheCollector->Remaining) ? Size : GCacheCollector->Remaining;
@@ -128,6 +133,11 @@ void Writer_CacheData(uint8* Data, uint32 Size)
 ////////////////////////////////////////////////////////////////////////////////
 void Writer_CacheOnConnect()
 {
+	if (GCacheCollector == nullptr)
+	{
+		return;
+	}
+
 	for (FCacheBuffer* Buffer = GCacheHeadBuffer; Buffer != nullptr; Buffer = Buffer->Next)
 	{
 		uint32 Used = Buffer->Size - Buffer->Remaining;
diff --git a/Engine/Source/Runtime/TraceLog/Private/Trace/Important/SharedBuffer.cpp b/Engine/Source/Runtime/TraceLog/Private/Trace/Important/SharedBuffer.cpp
--- a/Engine/Source/Runtime/TraceLog/Private/Trace/Important/SharedBuffer.cpp
+++ b/Engine/Source/Runtime/TraceLog/Private/Trace/Important/SharedBuffer.cpp
@@ -16,8 +16,6 @@ namespace Private {
 void*	Writer_MemoryAllocate(SIZE_T, uint32);
 void	Writer_MemoryFree(void*, uint32);
 void	Writer_CacheData(uint8*, uint32);
-void	Writer_InitializeCache();
-void	Writer_ShutdownCache();
 
 ////////////////////////////////////////////////////////////////////////////////
 static FSharedBuffer	GNullSharedBuffer	= { 0, FSharedBuffer::RefInit };
@@ -106,6 +104,10 @@ static void Writer_RetireSharedBufferImpl()
 	uint8* Data = (uint8*)GTailBuffer - GTailBuffer->Size + GTailPreSent;
 	if (uint32 SendSize = UPTRINT(GTailBuffer) - UPTRINT(Data) - GTailBuffer->Final)
 	{
+#if TRACE_PRIVATE_STATISTICS
+		GTraceStatistics.BytesTraced += SendSize;
+#endif
+
 		Writer_CacheData(Data, SendSize);
 	}
 
@@ -172,8 +174,6 @@ void Writer_UpdateSharedBuffers()
 ////////////////////////////////////////////////////////////////////////////////
 void Writer_InitializeSharedBuffers()
 {
-	Writer_InitializeCache();
-
 	FSharedBuffer* Buffer = Writer_CreateSharedBuffer();
 
 	GTailBuffer = Buffer;
@@ -185,7 +185,6 @@ void Writer_InitializeSharedBuffers()
 ////////////////////////////////////////////////////////////////////////////////
 void Writer_ShutdownSharedBuffers()
 {
-	Writer_ShutdownCache();
 }
 
 } // namespace Private
diff --git a/Engine/Source/Runtime/TraceLog/Private/Trace/TlsBuffer.cpp b/Engine/Source/Runtime/TraceLog/Private/Trace/TlsBuffer.cpp
--- a/Engine/Source/Runtime/TraceLog/Private/Trace/TlsBuffer.cpp
+++ b/Engine/Source/Runtime/TraceLog/Private/Trace/TlsBuffer.cpp
@@ -14,11 +14,12 @@ namespace Trace {
 namespace Private {
 
 ////////////////////////////////////////////////////////////////////////////////
-void			Writer_TailAppend(uint32, uint8* __restrict, uint32, bool);
-FWriteBuffer*	Writer_AllocateBlockFromPool();
-uint32			Writer_GetThreadId();
-void			Writer_FreeBlockListToPool(FWriteBuffer*, FWriteBuffer*);
-extern uint64	GStartCycle;
+void				Writer_TailAppend(uint32, uint8* __restrict, uint32, bool);
+FWriteBuffer*		Writer_AllocateBlockFromPool();
+uint32				Writer_GetThreadId();
+void				Writer_FreeBlockListToPool(FWriteBuffer*, FWriteBuffer*);
+extern uint64		GStartCycle;
+extern FStatistics	GTraceStatistics;
 
 
 
@@ -126,6 +127,10 @@ static bool Writer_DrainBuffer(uint32 ThreadId, FWriteBuffer* Buffer)
 	// Send as much as we can.
 	if (uint32 SizeToReap = uint32(Committed - Buffer->Reaped))
 	{
+#if TRACE_PRIVATE_STATISTICS
+		GTraceStatistics.BytesTraced += SizeToReap;
+#endif
+
 		bool bPartial = (Buffer->Partial == 1);
 		bPartial &= UPTRINT(Buffer->Reaped + Buffer->Size) == UPTRINT(Buffer);
 		Writer_TailAppend(ThreadId, Buffer->Reaped, SizeToReap, bPartial);
diff --git a/Engine/Source/Runtime/TraceLog/Private/Trace/Writer.cpp b/Engine/Source/Runtime/TraceLog/Private/Trace/Writer.cpp
--- a/Engine/Source/Runtime/TraceLog/Private/Trace/Writer.cpp
+++ b/Engine/Source/Runtime/TraceLog/Private/Trace/Writer.cpp
@@ -40,6 +40,8 @@ void			Writer_TailOnConnect();
 void			Writer_InitializeSharedBuffers();
 void			Writer_ShutdownSharedBuffers();
 void			Writer_UpdateSharedBuffers();
+void			Writer_InitializeCache();
+void			Writer_ShutdownCache();
 void			Writer_CacheOnConnect();
 void			Writer_InitializePool();
 void			Writer_ShutdownPool();
@@ -173,7 +175,7 @@ void* Writer_MemoryAllocate(SIZE_T Size, uint32 Alignment)
 	}
 
 #if TRACE_PRIVATE_STATISTICS
-	AtomicAddRelaxed(&GTraceStatistics.MemoryUsed, uint32(Size));
+	AtomicAddRelaxed(&GTraceStatistics.MemoryUsed, uint64(Size));
 #endif
 
 	return Ret;
@@ -219,7 +221,7 @@ void Writer_MemoryFree(void* Address, uint32 Size)
 #endif // TRACE_PRIVATE_STOMP
 
 #if TRACE_PRIVATE_STATISTICS
-	AtomicAddRelaxed(&GTraceStatistics.MemoryUsed, uint32(-int64(Size)));
+	AtomicAddRelaxed(&GTraceStatistics.MemoryUsed, uint64(-int64(Size)));
 #endif
 }
 
@@ -259,10 +261,6 @@ void Writer_SendData(uint32 ThreadId, uint8* __restrict Data, uint32 Size)
 {
 	static_assert(ETransport::Active == ETransport::TidPacketSync, "Active should be set to what the compiled code uses. It is used to track places that assume transport packet format");
 
-#if TRACE_PRIVATE_STATISTICS
-	GTraceStatistics.BytesTraced += Size;
-#endif
-
 	if (!GDataHandle)
 	{
 		return;
@@ -561,6 +559,7 @@ static void Writer_InternalShutdown()
 	Writer_ShutdownControl();
 	Writer_ShutdownPool();
 	Writer_ShutdownSharedBuffers();
+	Writer_ShutdownCache();
 	Writer_ShutdownTail();
 
 	GInitialized = false;
@@ -597,6 +596,11 @@ void Writer_Initialize(const FInitializeDesc& Desc)
 {
 	Writer_InitializeTail(Desc.TailSizeBytes);
 
+	if (Desc.bUseImportantCache)
+	{
+		Writer_InitializeCache();
+	}
+
 	if (Desc.bUseWorkerThread)
 	{
 		Writer_WorkerCreate();
diff --git a/Engine/Source/Runtime/TraceLog/Public/Trace/Trace.h b/Engine/Source/Runtime/TraceLog/Public/Trace/Trace.h
--- a/Engine/Source/Runtime/TraceLog/Public/Trace/Trace.h
+++ b/Engine/Source/Runtime/TraceLog/Public/Trace/Trace.h
@@ -25,6 +25,7 @@ struct FInitializeDesc
 {
 	uint32			TailSizeBytes		= 4 << 20;
 	bool			bUseWorkerThread	= true;
+	bool			bUseImportantCache	= true;
 };
 
 typedef void*		AllocFunc(SIZE_T, uint32);
@@ -34,9 +35,9 @@ struct FStatistics
 {
 	uint64		BytesSent;
 	uint64		BytesTraced;
-	uint32		MemoryUsed;
-	uint32		CacheUsed;
-	uint32		CacheWaste;
+	uint64		MemoryUsed;
+	uint32		CacheUsed;		// Important-marked events are
+	uint32		CacheWaste;		// stored in the cache.
 };
 
 UE_TRACE_API void	SetMemoryHooks(AllocFunc Alloc, FreeFunc Free) UE_TRACE_IMPL();
