Fixed narrowing warnings.

#rnx

diff --git a/Engine/Source/Runtime/TraceLog/Private/Trace/BlockPool.cpp b/Engine/Source/Runtime/TraceLog/Private/Trace/BlockPool.cpp
--- a/Engine/Source/Runtime/TraceLog/Private/Trace/BlockPool.cpp
+++ b/Engine/Source/Runtime/TraceLog/Private/Trace/BlockPool.cpp
@@ -62,7 +62,7 @@ static FPoolBlockList Writer_AddPageToPool(uint32 PageSize)
 	for (int i = 1, n = PageSize / GPoolBlockSize; ; ++i)
 	{
 		auto* Buffer = (FWriteBuffer*)Block;
-		Buffer->Size = BufferSize;
+		Buffer->Size = uint16(BufferSize);
 		if (i >= n)
 		{
 			break;
diff --git a/Engine/Source/Runtime/TraceLog/Private/Trace/Channel.cpp b/Engine/Source/Runtime/TraceLog/Private/Trace/Channel.cpp
--- a/Engine/Source/Runtime/TraceLog/Private/Trace/Channel.cpp
+++ b/Engine/Source/Runtime/TraceLog/Private/Trace/Channel.cpp
@@ -74,7 +74,7 @@ static uint32 GetChannelNameLength(const ANSICHAR* ChannelName)
 		}
 	}
 
-	return Len;
+	return uint32(Len);
 }
 
 
diff --git a/Engine/Source/Runtime/TraceLog/Private/Trace/Important/Cache.cpp b/Engine/Source/Runtime/TraceLog/Private/Trace/Important/Cache.cpp
--- a/Engine/Source/Runtime/TraceLog/Private/Trace/Important/Cache.cpp
+++ b/Engine/Source/Runtime/TraceLog/Private/Trace/Important/Cache.cpp
@@ -84,7 +84,7 @@ static void Writer_CacheCommit(const FCacheBuffer* Collector)
 	auto* Packet = (FTidPacketEncoded*)(GCacheActiveBuffer->Data + Used);
 	uint32 OutputSize = Encode(Collector->Data, InputSize, Packet->Data, EncodeMaxSize);
 
-	Packet->PacketSize = OutputSize + sizeof(FTidPacketEncoded);
+	Packet->PacketSize = uint16(OutputSize + sizeof(FTidPacketEncoded));
 	Packet->ThreadId = FTidPacketBase::EncodedMarker | uint16(ETransportTid::Importants);
 	Packet->DecodedSize = uint16(InputSize);
 
diff --git a/Engine/Source/Runtime/TraceLog/Private/Trace/Important/SharedBuffer.cpp b/Engine/Source/Runtime/TraceLog/Private/Trace/Important/SharedBuffer.cpp
--- a/Engine/Source/Runtime/TraceLog/Private/Trace/Important/SharedBuffer.cpp
+++ b/Engine/Source/Runtime/TraceLog/Private/Trace/Important/SharedBuffer.cpp
@@ -102,13 +102,13 @@ static void Writer_RetireSharedBufferImpl()
 {
 	// Send any unsent data.
 	uint8* Data = (uint8*)GTailBuffer - GTailBuffer->Size + GTailPreSent;
-	if (uint32 SendSize = UPTRINT(GTailBuffer) - UPTRINT(Data) - GTailBuffer->Final)
+	if (auto SendSize = UPTRINT(GTailBuffer) - UPTRINT(Data) - GTailBuffer->Final)
 	{
 #if TRACE_PRIVATE_STATISTICS
 		GTraceStatistics.BytesTraced += SendSize;
 #endif
 
-		Writer_CacheData(Data, SendSize);
+		Writer_CacheData(Data, uint32(SendSize));
 	}
 
 	FSharedBuffer* Temp = GTailBuffer->Next;
diff --git a/Engine/Source/Runtime/TraceLog/Private/Trace/Tail.cpp b/Engine/Source/Runtime/TraceLog/Private/Trace/Tail.cpp
--- a/Engine/Source/Runtime/TraceLog/Private/Trace/Tail.cpp
+++ b/Engine/Source/Runtime/TraceLog/Private/Trace/Tail.cpp
@@ -188,7 +188,7 @@ FTidPacketBase* FPacketRing::AppendImpl(uint32 InSize)
 
 	// Drop a packet from left.
 	auto* TidPacket = (FTidPacketBase*)(Data + Cursor);
-	TidPacket->PacketSize = InSize;
+	TidPacket->PacketSize = uint16(InSize);
 
 	Cursor = NextCursor;
 	return TidPacket;
diff --git a/Engine/Source/Runtime/TraceLog/Private/Trace/Writer.cpp b/Engine/Source/Runtime/TraceLog/Private/Trace/Writer.cpp
--- a/Engine/Source/Runtime/TraceLog/Private/Trace/Writer.cpp
+++ b/Engine/Source/Runtime/TraceLog/Private/Trace/Writer.cpp
@@ -346,7 +346,7 @@ void Writer_SendData(uint32 ThreadId, uint8* __restrict Data, uint32 Size)
 	Packet.ThreadId = FTidPacketBase::EncodedMarker;
 	Packet.ThreadId |= uint16(ThreadId & FTidPacketBase::ThreadIdMask);
 	Packet.DecodedSize = uint16(Size);
-	Packet.PacketSize = Encode(Data, Packet.DecodedSize, Packet.Data, sizeof(Packet.Data));
+	Packet.PacketSize = uint16(Encode(Data, Packet.DecodedSize, Packet.Data, sizeof(Packet.Data)));
 	Packet.PacketSize += sizeof(FTidPacketEncoded);
 
 	Writer_SendDataImpl(&Packet, Packet.PacketSize);
@@ -611,8 +611,8 @@ static void Writer_InternalInitializeImpl()
 	UE_TRACE_LOG($Trace, NewTrace, TraceLogChannel)
 		<< NewTrace.StartCycle(GStartCycle)
 		<< NewTrace.CycleFrequency(TimeGetFrequency())
-		<< NewTrace.Endian(0x524d)
-		<< NewTrace.PointerSize(sizeof(void*));
+		<< NewTrace.Endian(uint16(0x524d))
+		<< NewTrace.PointerSize(uint8(sizeof(void*)));
 }
 
 ////////////////////////////////////////////////////////////////////////////////
@@ -721,7 +721,7 @@ bool Writer_SendTo(const ANSICHAR* Host, uint32 Port)
 	Writer_InternalInitialize();
 
 	Port = Port ? Port : 1981;
-	UPTRINT DataHandle = TcpSocketConnect(Host, Port);
+	UPTRINT DataHandle = TcpSocketConnect(Host, uint16(Port));
 	if (!DataHandle)
 	{
 		return false;
