Fixed unused parameter/variable warnings

#rnx

diff --git a/Engine/Source/Runtime/TraceLog/Private/Trace/Channel.cpp b/Engine/Source/Runtime/TraceLog/Private/Trace/Channel.cpp
--- a/Engine/Source/Runtime/TraceLog/Private/Trace/Channel.cpp
+++ b/Engine/Source/Runtime/TraceLog/Private/Trace/Channel.cpp
@@ -233,7 +233,7 @@ FChannel* FChannel::FindChannel(const ANSICHAR* ChannelName)
 bool FChannel::Toggle(bool bEnabled)
 {
 	using namespace Private;
-	int64 OldRefCnt = AtomicAddRelaxed(&Enabled, bEnabled ? 1 : -1);
+	AtomicAddRelaxed(&Enabled, bEnabled ? 1 : -1);
 
 	UE_TRACE_LOG(Trace, ChannelToggle, TraceLogChannel)
 		<< ChannelToggle.Id(Name.Hash)
diff --git a/Engine/Source/Runtime/TraceLog/Private/Trace/Detail/Windows/WindowsTrace.cpp b/Engine/Source/Runtime/TraceLog/Private/Trace/Detail/Windows/WindowsTrace.cpp
--- a/Engine/Source/Runtime/TraceLog/Private/Trace/Detail/Windows/WindowsTrace.cpp
+++ b/Engine/Source/Runtime/TraceLog/Private/Trace/Detail/Windows/WindowsTrace.cpp
@@ -20,7 +20,7 @@ namespace Trace {
 namespace Private {
 
 ////////////////////////////////////////////////////////////////////////////////
-UPTRINT ThreadCreate(const ANSICHAR* Name, void (*Entry)())
+UPTRINT ThreadCreate(const ANSICHAR*, void (*Entry)())
 {
 	DWORD (WINAPI *WinApiThunk)(void*) = [] (void* Param) -> DWORD
 	{
diff --git a/Engine/Source/Runtime/TraceLog/Public/Trace/Detail/Trace.h b/Engine/Source/Runtime/TraceLog/Public/Trace/Detail/Trace.h
--- a/Engine/Source/Runtime/TraceLog/Public/Trace/Detail/Trace.h
+++ b/Engine/Source/Runtime/TraceLog/Public/Trace/Detail/Trace.h
@@ -116,7 +116,8 @@ class FChannel;
 #define TRACE_PRIVATE_LOG_PRELUDE(EnterFunc, LoggerName, EventName, ChannelsExpr, ...) \
 	if (TRACE_PRIVATE_CHANNELEXPR_IS_ENABLED(ChannelsExpr)) \
 		if (auto LogScope = F##LoggerName##EventName##Fields::LogScopeType::EnterFunc<F##LoggerName##EventName##Fields>(__VA_ARGS__)) \
-			if (const auto& __restrict EventName = *(F##LoggerName##EventName##Fields*)(&LogScope))
+			if (const auto& __restrict EventName = *(F##LoggerName##EventName##Fields*)(&LogScope)) \
+				((void)EventName),
 
 #define TRACE_PRIVATE_LOG_EPILOG() \
 				LogScope += LogScope
