Fixed implicit padding warning

#rnx

diff --git a/Engine/Source/Runtime/TraceLog/Private/Trace/Important/Cache.cpp b/Engine/Source/Runtime/TraceLog/Private/Trace/Important/Cache.cpp
--- a/Engine/Source/Runtime/TraceLog/Private/Trace/Important/Cache.cpp
+++ b/Engine/Source/Runtime/TraceLog/Private/Trace/Important/Cache.cpp
@@ -32,7 +32,7 @@ struct alignas(16) FCacheBuffer
 	};
 	uint32				Size;
 	uint32				Remaining;
-	uint32				_Unused;
+	uint32				_Unused[3];
 	uint32				Underflow; // For packet header
 	uint8				Data[];
 };
