Missing include

#rnx

diff --git a/Engine/Source/Runtime/TraceLog/Private/Trace/WriteBufferRedirect.h b/Engine/Source/Runtime/TraceLog/Private/Trace/WriteBufferRedirect.h
--- a/Engine/Source/Runtime/TraceLog/Private/Trace/WriteBufferRedirect.h
+++ b/Engine/Source/Runtime/TraceLog/Private/Trace/WriteBufferRedirect.h
@@ -2,6 +2,8 @@
 
 #pragma once
 
+#include "Trace/Detail/Writer.inl"
+
 namespace UE {
 namespace Trace {
 namespace Private {
