Multi-byte character constants are non-standard

#rnx

diff --git a/Engine/Source/Runtime/TraceLog/Private/Trace/Writer.cpp b/Engine/Source/Runtime/TraceLog/Private/Trace/Writer.cpp
--- a/Engine/Source/Runtime/TraceLog/Private/Trace/Writer.cpp
+++ b/Engine/Source/Runtime/TraceLog/Private/Trace/Writer.cpp
@@ -479,7 +479,7 @@ static bool Writer_UpdateConnection()
 	// Handshake.
 	struct FHandshake
 	{
-		uint32 Magic			= 'TRC2';
+		uint32 Magic			= '2' | ('C' << 8) | ('R' << 16) | ('T' << 24);
 		uint16 MetadataSize		= uint16(4); //  = sizeof(MetadataField0 + ControlPort)
 		uint16 MetadataField0	= uint16(sizeof(ControlPort) | (ControlPortFieldId << 8));
 		uint16 ControlPort		= uint16(Writer_GetControlPort());
