Deprecation warning fixes.

#rnx

diff --git a/Engine/Source/Runtime/TraceLog/Public/Trace/Detail/Important/ImportantLogScope.inl b/Engine/Source/Runtime/TraceLog/Public/Trace/Detail/Important/ImportantLogScope.inl
--- a/Engine/Source/Runtime/TraceLog/Public/Trace/Detail/Important/ImportantLogScope.inl
+++ b/Engine/Source/Runtime/TraceLog/Public/Trace/Detail/Important/ImportantLogScope.inl
@@ -27,7 +27,7 @@ TRACELOG_API FNextSharedBuffer				Writer_NextSharedBuffer(FSharedBuffer*, int32,
 template <class T>
 FORCENOINLINE FImportantLogScope FImportantLogScope::Enter(uint32 ArrayDataSize)
 {
-	static_assert(T::EventFlags & FEventInfo::Flag_MaybeHasAux, "Only important trace events with array-type fields need a size parameter to UE_TRACE_LOG()");
+	static_assert(uint32(T::EventFlags) & uint32(FEventInfo::Flag_MaybeHasAux), "Only important trace events with array-type fields need a size parameter to UE_TRACE_LOG()");
 
 	ArrayDataSize += sizeof(FAuxHeader) * T::EventProps_Meta::NumAuxFields;
 	ArrayDataSize += 1; // for AuxDataTerminal
@@ -45,7 +45,7 @@ FORCENOINLINE FImportantLogScope FImportantLogScope::Enter(uint32 ArrayDataSize)
 template <class T>
 inline FImportantLogScope FImportantLogScope::Enter()
 {
-	static_assert(!(T::EventFlags & FEventInfo::Flag_MaybeHasAux), "Important trace events with array-type fields must be traced with UE_TRACE_LOG(Logger, Event, Channel, ArrayDataSize)");
+	static_assert(!(uint32(T::EventFlags) & uint32(FEventInfo::Flag_MaybeHasAux)), "Important trace events with array-type fields must be traced with UE_TRACE_LOG(Logger, Event, Channel, ArrayDataSize)");
 
 	uint32 Size = T::GetSize();
 	uint32 Uid = T::GetUid() >> EKnownEventUids::_UidShift;
diff --git a/Engine/Source/Runtime/TraceLog/Public/Trace/Detail/Trace.h b/Engine/Source/Runtime/TraceLog/Public/Trace/Detail/Trace.h
--- a/Engine/Source/Runtime/TraceLog/Public/Trace/Detail/Trace.h
+++ b/Engine/Source/Runtime/TraceLog/Public/Trace/Detail/Trace.h
@@ -108,7 +108,7 @@ class FChannel;
 		explicit operator bool () const { return true; } \
 		enum { EventFlags = PartialEventFlags|(EventProps_Meta::NumAuxFields ? UE::Trace::Private::FEventInfo::Flag_MaybeHasAux : 0), }; \
 		static_assert( \
-			!bIsImportant || (EventFlags & UE::Trace::Private::FEventInfo::Flag_NoSync), \
+			!bIsImportant || (uint32(EventFlags) & uint32(UE::Trace::Private::FEventInfo::Flag_NoSync)), \
 			"Trace events flagged as Important events must be marked NoSync" \
 		); \
 	};
