Removed a PLATFORM_X define in favour of the setting UE_TRACE_ENABLED in the appropriate .Build.cs file instead.

#rnx

diff --git a/Engine/Source/Runtime/TraceLog/Public/Trace/Config.h b/Engine/Source/Runtime/TraceLog/Public/Trace/Config.h
--- a/Engine/Source/Runtime/TraceLog/Public/Trace/Config.h
+++ b/Engine/Source/Runtime/TraceLog/Public/Trace/Config.h
@@ -6,7 +6,7 @@
 
 #if !defined(UE_TRACE_ENABLED)
 #	if !UE_BUILD_SHIPPING && !IS_PROGRAM
-#		if PLATFORM_WINDOWS || PLATFORM_UNIX || PLATFORM_APPLE || PLATFORM_SWITCH || PLATFORM_ANDROID || PLATFORM_HOLOLENS
+#		if PLATFORM_WINDOWS || PLATFORM_UNIX || PLATFORM_APPLE || PLATFORM_ANDROID || PLATFORM_HOLOLENS
 #			define UE_TRACE_ENABLED	1
 #		endif
 #	endif
