Can't append the tail as alloc/free can get called from any thread and the tail is assumed to only be used by the worker thread (or update call). Immediately appending to the tail also desynchronises the contiguous assumption of serials within a single thread. Ultimately it's a chicken-and-the-egg problem - memory needs to be allocated to store the trace events that memory allocation traces. A redirect is kept around just in case, but its use is considered an error.

Bonus; analysis can get going quicker if a sync is added directly after the tails been sent to a new connection.

diff --git a/Engine/Source/Runtime/TraceLog/Private/Trace/Writer.cpp b/Engine/Source/Runtime/TraceLog/Private/Trace/Writer.cpp
--- a/Engine/Source/Runtime/TraceLog/Private/Trace/Writer.cpp
+++ b/Engine/Source/Runtime/TraceLog/Private/Trace/Writer.cpp
@@ -40,7 +40,6 @@ int32			Encode(const void*, int32, void*, int32);
 void			Writer_SendData(uint32, uint8* __restrict, uint32);
 void			Writer_InitializeTail(int32);
 void			Writer_ShutdownTail();
-void			Writer_TailAppend(uint32, uint8* __restrict, uint32, bool=false);
 void			Writer_TailOnConnect();
 void			Writer_InitializeSharedBuffers();
 void			Writer_ShutdownSharedBuffers();
@@ -137,7 +136,7 @@ void Writer_MemorySetHooks(decltype(AllocHook) Alloc, decltype(FreeHook) Free)
 ////////////////////////////////////////////////////////////////////////////////
 void* Writer_MemoryAllocate(SIZE_T Size, uint32 Alignment)
 {
-	TWriteBufferRedirect<6 << 10> TraceData;
+	TWriteBufferRedirect<1 << 10> TraceData;
 
 	void* Ret = nullptr;
 
@@ -176,12 +175,6 @@ void* Writer_MemoryAllocate(SIZE_T Size, uint32 Alignment)
 	}
 #endif // TRACE_PRIVATE_STOMP
 
-	if (TraceData.GetSize())
-	{
-		uint32 ThreadId = Writer_GetThreadId();
-		Writer_TailAppend(ThreadId, TraceData.GetData(), TraceData.GetSize());
-	}
-
 #if TRACE_PRIVATE_STATISTICS
 	AtomicAddRelaxed(&GTraceStatistics.MemoryUsed, uint64(Size));
 #endif
@@ -206,7 +199,7 @@ void Writer_MemoryFree(void* Address, uint32 Size)
 	DWORD Unused;
 	VirtualProtect(MemInfo.BaseAddress, MemInfo.RegionSize, PAGE_READONLY, &Unused);
 #else // TRACE_PRIVATE_STOMP
-	TWriteBufferRedirect<6 << 10> TraceData;
+	TWriteBufferRedirect<1 << 10> TraceData;
 
 	if (FreeHook != nullptr)
 	{
@@ -220,12 +213,6 @@ void Writer_MemoryFree(void* Address, uint32 Size)
 		free(Address);
 #endif
 	}
-
-	if (TraceData.GetSize())
-	{
-		uint32 ThreadId = Writer_GetThreadId();
-		Writer_TailAppend(ThreadId, TraceData.GetData(), TraceData.GetSize());
-	}
 #endif // TRACE_PRIVATE_STOMP
 
 #if TRACE_PRIVATE_STATISTICS
