Script to create a single-file header version of Trace with a CPU scope implementation.

#rnx

diff --git a/Engine/Source/Runtime/TraceLog/create_standalone.py b/Engine/Source/Runtime/TraceLog/create_standalone.py
new file mode 100644
--- /dev/null
+++ b/Engine/Source/Runtime/TraceLog/create_standalone.py
@@ -0,0 +1,150 @@
+import os
+import re
+import argparse
+from pathlib import Path
+
+def _spam_header(*args, **kwargs):
+	_spam("\x1b[96m##", *args, "\x1b[0m", **kwargs)
+
+def _spam(*args, **kwargs):
+	print(*args, **kwargs)
+
+class _SourceFile(object):
+	def __init__(self, path):
+		self.path = path
+		self.is_include = (path.suffix == ".h")
+		self.lines = []
+		self.deps = []
+
+	def __eq__(self, rhs):
+		return rhs.samefile(self.path)
+
+	def __hash__(self):
+		return hash(self.path)
+
+def _parse_include(line):
+	m = re.match(r'\s*#\s*include\s+\"([^"]+)\"', line)
+	return None if not m else Path(m.group(1))
+
+def _exclude_line(line):
+	line = line.strip()
+
+	if "#pragma once" in line:	return True
+	if line.startswith("//"):	return True
+	if not line:				return True
+
+def _main(src_dir, dest_dir):
+	dest_dir = dest_dir.resolve()
+	os.chdir(Path(__file__).parent)
+
+	_spam_header("Context")
+	_spam("Source dir:", os.getcwd())
+	_spam("Dest dir:", dest_dir)
+
+	_spam_header("Gathering source files")
+	files  = [x for x in src_dir.glob("Public/**/*")]
+	files += [x for x in src_dir.glob("Private/**/*")]
+	files  = [x for x in files if x.suffix in (".h", ".cpp", ".inl")]
+	_spam("Found %d files" % len(files))
+
+	# Collect lines of each source file, filtering local includes
+	source_files = []
+	for file in files:
+		source_file = _SourceFile(file)
+		source_files.append(source_file)
+
+		for line in file.open("rt"):
+			include = _parse_include(line)
+			if not include:
+				if not _exclude_line(line):
+					source_file.lines.append(line)
+				continue
+
+			is_local = False
+
+			for include_dir in (file.parent, src_dir / "Public"):
+				candidate = include_dir / include
+				if candidate.is_file():
+					is_local = True
+					break
+
+			if is_local:
+				source_file.deps.append(candidate)
+
+	_spam_header("Topological sort")
+
+	# Hook up dependencies
+	_spam("Resolving include dependencies")
+	for source_file in source_files:
+		new_deps = []
+		for dep in source_file.deps:
+			index = source_files.index(dep)
+			dep = source_files[index]
+			new_deps.append(dep)
+		source_file.deps = new_deps
+
+	# Topologically sort by dependencies
+	_spam("Sorting")
+	visited = set()
+	topo_sorted = []
+	def visit(source_file):
+		if source_file in visited:
+			return
+
+		for x in source_file.deps:
+			visit(x)
+
+		topo_sorted.append(source_file)
+		visited.add(source_file)
+
+	for x in source_files:
+		visit(x)
+
+	# Stable sort to put .cpps last
+	ext_key = lambda x: 1 if x.path.suffix == ".cpp" else 0
+	source_files = sorted(topo_sorted, key=ext_key)
+
+	# Exclude LZ4
+	source_files = [x for x in source_files if "lz4" not in x.path.name]
+
+	# Add prologue and epilogue files
+	prologue = _SourceFile(Path("standalone_prologue.h"))
+	prologue.lines = [x for x in prologue.path.open("rt") if x]
+	source_files.insert(0, prologue)
+
+	epilogue = _SourceFile(Path("standalone_epilogue.h"))
+	epilogue.lines = [x for x in epilogue.path.open("rt") if x]
+	source_files.append(epilogue)
+
+	# Write the output file
+	_spam_header("Output")
+	_spam(dest_dir / "trace.h")
+	cpp_started = False
+	with (dest_dir / "trace.h").open("wt") as out:
+		print("// Copyright Epic Games, Inc. All Rights Reserved.", file=out)
+		print("#pragma once", file=out)
+		for source_file in source_files:
+			if not cpp_started:
+				if source_file.path.suffix == ".cpp":
+					print("#ifdef TRACE_IMPLEMENT", file=out)
+					cpp_started = True
+			elif source_file.path.suffix != ".cpp":
+				print("#endif // TRACE_IMPLEMENT", file=out)
+				cpp_started = False
+
+			print("/* {{{1", source_file.path.name, "*/", file=out)
+			print(file=out)
+			for line in source_file.lines:
+				out.write(line)
+
+	_spam("...done!")
+
+def main():
+	desc = "Amalgamate TraceLog into a standalone single-file library"
+	parser = argparse.ArgumentParser(description=desc)
+	parser.add_argument("outdir", help="Directory to write output file(s) to")
+	args = parser.parse_args()
+	return _main(Path(__file__).parent, Path(args.outdir))
+
+if __name__ == "__main__":
+	raise SystemExit(main())
diff --git a/Engine/Source/Runtime/TraceLog/standalone_epilogue.h b/Engine/Source/Runtime/TraceLog/standalone_epilogue.h
new file mode 100644
--- /dev/null
+++ b/Engine/Source/Runtime/TraceLog/standalone_epilogue.h
@@ -0,0 +1,275 @@
+// Copyright Epic Games, Inc. All Rights Reserved.
+
+#if PLATFORM_WINDOWS
+#	pragma warning(pop)
+#endif
+
+#if TRACE_UE_COMPAT_LAYER
+
+#if PLATFORM_WINDOWS
+#	if defined(UNICODE) || defined(_UNICODE)
+#		undef TEXT
+#		undef TCHAR
+#		define TEXT(x)	L##x
+#	endif
+#endif
+
+#endif // TRACE_UE_COMPAT_LAYER
+
+#include <string_view>
+
+#define TRACE_EVENT_DEFINE			UE_TRACE_EVENT_DEFINE
+#define TRACE_EVENT_BEGIN			UE_TRACE_EVENT_BEGIN
+#define TRACE_EVENT_BEGIN_EXTERN	UE_TRACE_EVENT_BEGIN_EXTERN
+#define TRACE_EVENT_FIELD			UE_TRACE_EVENT_FIELD
+#define TRACE_EVENT_END				UE_TRACE_EVENT_END
+#define TRACE_LOG					UE_TRACE_LOG
+#define TRACE_LOG_SCOPED			UE_TRACE_LOG_SCOPED
+#define TRACE_LOG_SCOPED_T			UE_TRACE_LOG_SCOPED_T
+#define TRACE_CHANNEL				UE_TRACE_CHANNEL
+#define TRACE_CHANNEL_EXTERN		UE_TRACE_CHANNEL_EXTERN
+#define TRACE_CHANNEL_DEFINE		UE_TRACE_CHANNEL_DEFINE
+
+namespace trace = UE::Trace;
+
+#define TRACE_PRIVATE_CONCAT_(x, y)		x##y
+#define TRACE_PRIVATE_CONCAT(x, y)		TRACE_PRIVATE_CONCAT_(x, y)
+#define TRACE_PRIVATE_UNIQUE_VAR(name)	TRACE_PRIVATE_CONCAT($trace_##name, __LINE__)
+
+TRACE_CHANNEL_EXTERN(CpuChannel)
+
+namespace UE {
+namespace Trace {
+
+struct TraceCpuScope
+{
+			~TraceCpuScope();
+	void	Enter(int ScopeId);
+	int		_ScopeId = 0;
+};
+
+int ScopeNew(const std::string_view& Name);
+
+} // namespace Trace
+} // namespace UE
+
+#define TRACE_CPU_SCOPE(name) \
+	using namespace std::literals; \
+	trace::TraceCpuScope TRACE_PRIVATE_UNIQUE_VAR(cpu_scope); \
+	if (CpuChannel) { \
+		static int TRACE_PRIVATE_UNIQUE_VAR(scope_id); \
+		if (0 == TRACE_PRIVATE_UNIQUE_VAR(scope_id)) \
+			TRACE_PRIVATE_UNIQUE_VAR(scope_id) = trace::ScopeNew(name##sv); \
+		TRACE_PRIVATE_UNIQUE_VAR(cpu_scope).Enter(TRACE_PRIVATE_UNIQUE_VAR(scope_id)); \
+	} \
+	do {} while (0)
+
+#if TRACE_IMPLEMENT
+
+////////////////////////////////////////////////////////////////////////////////
+TRACE_CHANNEL_DEFINE(CpuChannel)
+
+TRACE_EVENT_BEGIN(CpuProfiler, EventSpec, NoSync|Important)
+	TRACE_EVENT_FIELD(uint32, Id)
+	TRACE_EVENT_FIELD(UE::Trace::AnsiString, Name)
+TRACE_EVENT_END()
+
+TRACE_EVENT_BEGIN(CpuProfiler, EventBatch, NoSync)
+	TRACE_EVENT_FIELD(uint8[], Data)
+TRACE_EVENT_END()
+
+namespace UE {
+namespace Trace {
+namespace Private {
+
+static int32_t encode32_7bit(int32_t value, void* __restrict out)
+{
+	// Calculate the number of bytes
+#if 0
+	int32_t msb_test = (value << sizeof(value)) | 0x10;
+#if _MSC_VER
+	unsigned long bit_index;
+	_BitScanReverse(&bit_index, msb_test);
+#else
+	int32_t leading_zeros = __builtin_clz(msb_test);
+	int32_t bit_index = ((sizeof(value) * 8) - 1) - leading_zeros;
+#endif
+	int32_t length = (bit_index + 3) / 7;
+#else
+	int32_t length = 1;
+	length += (value >= (1 <<  7));
+	length += (value >= (1 << 14));
+	length += (value >= (1 << 21));
+#endif
+
+	// Add a gap every eigth bit for the continuations
+	int32_t ret = value;
+	ret = (ret & 0x0000'3fff) | ((ret & 0x0fff'c000) << 2);
+	ret = (ret & 0x007f'007f) | ((ret & 0x3f80'3f80) << 1);
+
+	// Set the bits indicating another byte follows
+	int32_t continuations = 0x0080'8080;
+	continuations >>= (sizeof(value) - length) * 8;
+	ret |= continuations;
+
+	::memcpy(out, &ret, sizeof(value));
+
+	return length;
+}
+
+static int32_t encode64_7bit(int64_t value, void* __restrict out)
+{
+	// Calculate the output length
+#if 0
+	int64_t msb_test = (value << sizeof(value)) | 0x100ull;
+#if _MSC_VER
+	unsigned long bit_index;
+	_BitScanReverse64(&bit_index, msb_test);
+#else
+	int32_t leading_zeros = __builtin_clzll(msb_test);
+	int32_t bit_index = ((sizeof(value) * 8) - 1) - leading_zeros;
+#endif
+	int32_t length = (bit_index - 1) / 7;
+#else
+	uint32_t length = 1;
+	length += (value >= (1ll <<  7));
+	length += (value >= (1ll << 14));
+	length += (value >= (1ll << 21));
+	length += (value >= (1ll << 28));
+	length += (value >= (1ll << 35));
+	length += (value >= (1ll << 42));
+	length += (value >= (1ll << 49));
+#endif
+
+	// Add a gap every eigth bit for the continuations
+	int64_t ret = value;
+	ret = (ret & 0x0000'0000'0fff'ffffull) | ((ret & 0x00ff'ffff'f000'0000ull) << 4);
+	ret = (ret & 0x0000'3fff'0000'3fffull) | ((ret & 0x0fff'c000'0fff'c000ull) << 2);
+	ret = (ret & 0x007f'007f'007f'007full) | ((ret & 0x3f80'3f80'3f80'3f80ull) << 1);
+
+	// Set the bits indicating another byte follows
+	int64_t continuations = 0x0080'8080'8080'8080ull;
+	continuations >>= (sizeof(value) - length) * 8;
+	ret |= continuations;
+
+	::memcpy(out, &ret, sizeof(value));
+
+	return length;
+}
+
+////////////////////////////////////////////////////////////////////////////////
+class ThreadBuffer
+{
+public:
+	static void	Enter(uint64_t Timestamp, uint32_t ScopeId) { TlsInstance.EnterImpl(Timestamp, ScopeId); }
+	static void	Leave(uint64_t Timestamp)					{ TlsInstance.LeaveImpl(Timestamp); }
+
+private:
+				~ThreadBuffer();
+	void		Flush(bool Force);
+	void		EnterImpl(uint64_t Timestamp, uint32_t ScopeId);
+	void		LeaveImpl(uint64_t Timestamp);
+	enum
+	{
+		BufferSize	= 256,
+		Overflow	= 16,
+		EnterLsb	= 1,
+		LeaveLsb	= 0,
+	};
+	uint64_t	PrevTimestamp = 0;
+	uint8_t*	Cursor = Buffer;
+	uint8_t		Buffer[BufferSize];
+
+	static thread_local ThreadBuffer TlsInstance;
+};
+
+thread_local ThreadBuffer ThreadBuffer::TlsInstance;
+
+////////////////////////////////////////////////////////////////////////////////
+ThreadBuffer::~ThreadBuffer()
+{
+	Flush(true);
+}
+
+////////////////////////////////////////////////////////////////////////////////
+void ThreadBuffer::Flush(bool Force)
+{
+	if (!Force && (Cursor <= (Buffer + BufferSize - Overflow)))
+		return;
+
+	TRACE_LOG(CpuProfiler, EventBatch, true)
+		<< EventBatch.Data(Buffer, uint32(ptrdiff_t(Cursor - Buffer)));
+
+	PrevTimestamp = 0;
+	Cursor = Buffer;
+}
+
+////////////////////////////////////////////////////////////////////////////////
+void ThreadBuffer::EnterImpl(uint64_t Timestamp, uint32_t ScopeId)
+{
+	Timestamp -= PrevTimestamp;
+	PrevTimestamp += Timestamp;
+	Cursor += encode64_7bit((Timestamp) << 1 | EnterLsb, Cursor);
+	Cursor += encode32_7bit(ScopeId, Cursor);
+
+	Flush(false);
+}
+
+////////////////////////////////////////////////////////////////////////////////
+void ThreadBuffer::LeaveImpl(uint64_t Timestamp)
+{
+	Timestamp -= PrevTimestamp;
+	PrevTimestamp += Timestamp;
+	Cursor += encode64_7bit((Timestamp << 1) | LeaveLsb, Cursor);
+
+	Flush(false);
+}
+
+} // namespace Private
+
+
+
+////////////////////////////////////////////////////////////////////////////////
+int ScopeNew(const std::string_view& Name)
+{
+	static int volatile NextSpecId = 1;
+	int SpecId = Private::AtomicAddRelaxed(&NextSpecId, 1);
+
+	uint32 NameSize = uint32(Name.size());
+	TRACE_LOG(CpuProfiler, EventSpec, true, NameSize)
+		<< EventSpec.Id(uint32(SpecId))
+		<< EventSpec.Name(Name.data(), NameSize);
+
+	return SpecId;
+}
+
+
+
+////////////////////////////////////////////////////////////////////////////////
+TraceCpuScope::~TraceCpuScope()
+{
+	using namespace Private;
+
+	if (!_ScopeId)
+		return;
+
+	uint64 Timestamp = TimeGetTimestamp();
+	ThreadBuffer::Leave(Timestamp);
+}
+
+////////////////////////////////////////////////////////////////////////////////
+void TraceCpuScope::Enter(int ScopeId)
+{
+	using namespace Private;
+
+	_ScopeId = ScopeId;
+	uint64 Timestamp = TimeGetTimestamp();
+	ThreadBuffer::Enter(Timestamp, ScopeId);
+}
+
+} // namespace Trace
+} // namespace UE
+
+#endif // TRACE_IMPLEMENT
+
+/* vim: set noet foldlevel=1 foldmethod=marker : */
diff --git a/Engine/Source/Runtime/TraceLog/standalone_prologue.h b/Engine/Source/Runtime/TraceLog/standalone_prologue.h
new file mode 100644
--- /dev/null
+++ b/Engine/Source/Runtime/TraceLog/standalone_prologue.h
@@ -0,0 +1,130 @@
+// Copyright Epic Games, Inc. All Rights Reserved.
+
+#pragma once
+
+#if !defined(TRACE_UE_COMPAT_LAYER)
+#	define TRACE_UE_COMPAT_LAYER	(!__UNREAL__)
+#endif
+
+#if TRACE_UE_COMPAT_LAYER
+
+#include <cstddef>
+#include <cstdint>
+#include <utility>
+
+#ifdef _WIN32
+#	define PLATFORM_WINDOWS			1
+#elif defined(__linux__)
+#	define PLATFORM_LINUX			1
+#elif defined(__APPLE__)
+#	define PLATFORM_MAC				1
+#endif
+
+#if defined(__amd64__) || defined(_M_X64)
+#	define PLATFORM_CPU_X86_FAMILY	1
+#	define PLATFORM_64BITS			1
+#elif defined(__arm64__) || defined(_M_ARM64)
+#	define PLATFORM_CPU_ARM_FAMILY	1
+#	define PLATFORM_64BITS			1
+#else
+#	error Unknown architecture
+#endif
+
+#if PLATFORM_WINDOWS
+#	if !defined(WIN32_LEAN_AND_MEAN)
+#		define WIN32_LEAN_AND_MEAN
+#	endif
+#	if !defined(NOGDI)
+#		define NOGDI
+#	endif
+#	if !defined(NOMINMAX)
+#		define NOMINMAX
+#	endif
+#	include <Windows.h>
+#endif
+
+// types
+using uint8  = uint8_t;
+using uint16 = uint16_t;
+using uint32 = uint32_t;
+using uint64 = uint64_t;
+
+using int8	= int8_t;
+using int16 = int16_t;
+using int32 = int32_t;
+using int64 = int64_t;
+
+using UPTRINT = uintptr_t;
+using PTRINT  = intptr_t;
+
+using SIZE_T = size_t;
+
+#if PLATFORM_WINDOWS
+#	undef TEXT
+#endif
+#define TEXT(x)	x
+#define TCHAR	ANSICHAR
+using ANSICHAR	= char;
+using WIDECHAR	= wchar_t;
+
+// keywords
+#if defined(_MSC_VER)
+#	define FORCENOINLINE	__declspec(noinline)
+#	define FORCEINLINE		__forceinline
+#else
+#	define FORCENOINLINE	inline __attribute__((noinline))
+#	define FORCEINLINE		inline __attribute__((always_inline))
+#endif
+
+#if defined(_MSC_VER)
+#	define LIKELY(x)		x
+#	define UNLIKELY(x)		x
+#else
+#	define LIKELY(x)		__builtin_expect(!!(x), 1)
+#	define UNLIKELY(x)		__builtin_expect(!!(x), 0)
+#endif
+
+#define UE_ARRAY_COUNT(x)	(sizeof(x) / sizeof(x[0]))
+
+// so/dll
+#if defined(TRACE_DLL_EXPORT)
+#	if PLATFORM_WINDOWS && defined(TRACE_DLL_EXPORT)
+#		if defined(TRACE_IMPLEMENT)
+#			define TRACELOG_API __declspec(dllexport)
+#		else
+#			define TRACELOG_API __declspec(dllimport)
+#		endif
+#	else
+#		define TRACELOG_API		__attribute__ ((visibility ("default")))
+#	endif
+#else
+#	define TRACELOG_API
+#endif
+
+// misc defines
+#define TRACE_ENABLED					1
+#define UE_TRACE_ENABLED				TRACE_ENABLED
+#define TRACE_PRIVATE_CONTROL_ENABLED	0
+#define TRACE_PRIVATE_EXTERNAL_LZ4		1
+#define PLATFORM_CACHE_LINE_SIZE		64
+#define THIRD_PARTY_INCLUDES_START
+#define THIRD_PARTY_INCLUDES_END
+
+// api
+template <typename T>
+inline auto Forward(T t)
+{
+	return std::forward<T>(t);
+}
+
+#endif // TRACE_UE_COMPAT_LAYER
+
+#include <cstring>
+#include "lz4.h"
+
+#if PLATFORM_WINDOWS
+#	pragma warning(push)
+#	pragma warning(disable : 4200) // zero-sized arrays
+#	pragma warning(disable : 4201) // anonymous structs
+#	pragma warning(disable : 4127) // conditional expr. is constant
+#endif
