Fixed incorrect preprocessor statement that wrapped globbed .cpp files.

#rnx

diff --git a/Engine/Source/Runtime/TraceLog/create_standalone.py b/Engine/Source/Runtime/TraceLog/create_standalone.py
--- a/Engine/Source/Runtime/TraceLog/create_standalone.py
+++ b/Engine/Source/Runtime/TraceLog/create_standalone.py
@@ -126,7 +126,7 @@ def _main(src_dir, dest_dir):
 		for source_file in source_files:
 			if not cpp_started:
 				if source_file.path.suffix == ".cpp":
-					print("#ifdef TRACE_IMPLEMENT", file=out)
+					print("#if TRACE_IMPLEMENT", file=out)
 					cpp_started = True
 			elif source_file.path.suffix != ".cpp":
 				print("#endif // TRACE_IMPLEMENT", file=out)
