Updated prologue so that all macros that the glob source uses are always defined when preprocessing.

#rnx

diff --git a/Engine/Source/Runtime/TraceLog/standalone_prologue.h b/Engine/Source/Runtime/TraceLog/standalone_prologue.h
--- a/Engine/Source/Runtime/TraceLog/standalone_prologue.h
+++ b/Engine/Source/Runtime/TraceLog/standalone_prologue.h
@@ -3,33 +3,59 @@
 #pragma once
 
 #if !defined(TRACE_UE_COMPAT_LAYER)
-#	define TRACE_UE_COMPAT_LAYER	(!__UNREAL__)
+#	if defined(__UNREAL__) && __UNREAL__
+#		define TRACE_UE_COMPAT_LAYER	0
+#	else
+#		define TRACE_UE_COMPAT_LAYER	1
+#	endif
 #endif
 
 #if TRACE_UE_COMPAT_LAYER
 
-#include <cstddef>
-#include <cstdint>
-#include <utility>
+// platform defines
+#define PLATFORM_ANDROID	0
+#define PLATFORM_APPLE		0
+#define PLATFORM_HOLOLENS	0
+#define PLATFORM_MAC		0
+#define PLATFORM_UNIX		0
+#define PLATFORM_WINDOWS	0
 
 #ifdef _WIN32
+#	undef PLATFORM_WINDOWS
 #	define PLATFORM_WINDOWS			1
 #elif defined(__linux__)
-#	define PLATFORM_LINUX			1
+#	undef  PLATFORM_UNIX
+#	define PLATFORM_UNIX			1
 #elif defined(__APPLE__)
+#	undef  PLATFORM_MAC
 #	define PLATFORM_MAC				1
+#	undef  PLATFORM_APPLE
+#	define PLATFORM_APPLE			1
 #endif
 
+#if PLATFORM_MAC
+#endif
+
+// arch defines
 #if defined(__amd64__) || defined(_M_X64)
 #	define PLATFORM_CPU_X86_FAMILY	1
+#	define PLATFORM_CPU_ARM_FAMILY	0
 #	define PLATFORM_64BITS			1
+#	define PLATFORM_CACHE_LINE_SIZE	64
 #elif defined(__arm64__) || defined(_M_ARM64)
+#	define PLATFORM_CPU_X86_FAMILY	0
 #	define PLATFORM_CPU_ARM_FAMILY	1
 #	define PLATFORM_64BITS			1
+#	define PLATFORM_CACHE_LINE_SIZE	64
 #else
 #	error Unknown architecture
 #endif
 
+// external includes
+#include <cstddef>
+#include <cstdint>
+#include <utility>
+
 #if PLATFORM_WINDOWS
 #	if !defined(WIN32_LEAN_AND_MEAN)
 #		define WIN32_LEAN_AND_MEAN
@@ -89,7 +115,7 @@ using WIDECHAR	= wchar_t;
 // so/dll
 #if defined(TRACE_DLL_EXPORT)
 #	if PLATFORM_WINDOWS && defined(TRACE_DLL_EXPORT)
-#		if defined(TRACE_IMPLEMENT)
+#		if TRACE_IMPLEMENT
 #			define TRACELOG_API __declspec(dllexport)
 #		else
 #			define TRACELOG_API __declspec(dllimport)
@@ -101,14 +127,18 @@ using WIDECHAR	= wchar_t;
 #	define TRACELOG_API
 #endif
 
+#if !defined(IS_MONOLITHIC)
+#	define IS_MONOLITHIC				0
+#endif
+
 // misc defines
+#define THIRD_PARTY_INCLUDES_END
+#define THIRD_PARTY_INCLUDES_START
 #define TRACE_ENABLED					1
-#define UE_TRACE_ENABLED				TRACE_ENABLED
 #define TRACE_PRIVATE_CONTROL_ENABLED	0
 #define TRACE_PRIVATE_EXTERNAL_LZ4		1
-#define PLATFORM_CACHE_LINE_SIZE		64
-#define THIRD_PARTY_INCLUDES_START
-#define THIRD_PARTY_INCLUDES_END
+#define UE_TRACE_ENABLED				TRACE_ENABLED
+#define UE_BUILD_SHIPPING				1
 
 // api
 template <typename T>
