Fix Linux crash when connecting to an incorrect host address.

diff --git a/Engine/Source/Runtime/TraceLog/Private/Trace/Detail/Unix/UnixTrace.cpp b/Engine/Source/Runtime/TraceLog/Private/Trace/Detail/Unix/UnixTrace.cpp
--- a/Engine/Source/Runtime/TraceLog/Private/Trace/Detail/Unix/UnixTrace.cpp
+++ b/Engine/Source/Runtime/TraceLog/Private/Trace/Detail/Unix/UnixTrace.cpp
@@ -93,7 +93,13 @@ UPTRINT TcpSocketConnect(const ANSICHAR* Host, uint16 Port)
 {
 	struct FAddrInfoPtr
 	{
-					~FAddrInfoPtr()	{ freeaddrinfo(Value); }
+		~FAddrInfoPtr()	
+		{
+			if (Value != nullptr)
+			{
+				freeaddrinfo(Value);
+			}
+		}
 		addrinfo*	operator -> ()	{ return Value; }
 		addrinfo**	operator & ()	{ return &Value; }
 		addrinfo*	Value;
@@ -106,6 +112,7 @@ UPTRINT TcpSocketConnect(const ANSICHAR* Host, uint16 Port)
 	Hints.ai_protocol = IPPROTO_TCP;
 	if (getaddrinfo(Host, nullptr, &Hints, &Info))
 	{
+		Info.Value = nullptr;
 		return 0;
 	}
 
