Separators for clarity

#rnx

diff --git a/Engine/Source/Runtime/TraceLog/create_standalone.py b/Engine/Source/Runtime/TraceLog/create_standalone.py
--- a/Engine/Source/Runtime/TraceLog/create_standalone.py
+++ b/Engine/Source/Runtime/TraceLog/create_standalone.py
@@ -3,12 +3,14 @@ import re
 import argparse
 from pathlib import Path
 
+#-------------------------------------------------------------------------------
 def _spam_header(*args, **kwargs):
 	_spam("\x1b[96m##", *args, "\x1b[0m", **kwargs)
 
 def _spam(*args, **kwargs):
 	print(*args, **kwargs)
 
+#-------------------------------------------------------------------------------
 class _SourceFile(object):
 	def __init__(self, path):
 		self.path = path
@@ -22,6 +24,7 @@ class _SourceFile(object):
 	def __hash__(self):
 		return hash(self.path)
 
+#-------------------------------------------------------------------------------
 def _parse_include(line):
 	m = re.match(r'\s*#\s*include\s+\"([^"]+)\"', line)
 	return None if not m else Path(m.group(1))
@@ -33,6 +36,7 @@ def _exclude_line(line):
 	if line.startswith("//"):	return True
 	if not line:				return True
 
+#-------------------------------------------------------------------------------
 def _main(src_dir, dest_dir):
 	dest_dir = dest_dir.resolve()
 	os.chdir(Path(__file__).parent)
