Added a --thin option to ease prologue/epilogue iteration

#rnx

diff --git a/Engine/Source/Runtime/TraceLog/create_standalone.py b/Engine/Source/Runtime/TraceLog/create_standalone.py
--- a/Engine/Source/Runtime/TraceLog/create_standalone.py
+++ b/Engine/Source/Runtime/TraceLog/create_standalone.py
@@ -37,7 +37,7 @@ def _exclude_line(line):
 	if not line:				return True
 
 #-------------------------------------------------------------------------------
-def _main(src_dir, dest_dir):
+def _main(src_dir, dest_dir, thin):
 	dest_dir = dest_dir.resolve()
 	os.chdir(Path(__file__).parent)
 
@@ -112,13 +112,13 @@ def _main(src_dir, dest_dir):
 	source_files = [x for x in source_files if "lz4" not in x.path.name]
 
 	# Add prologue and epilogue files
-	prologue = _SourceFile(Path("standalone_prologue.h"))
-	prologue.lines = [x for x in prologue.path.open("rt") if x]
-	source_files.insert(0, prologue)
-
-	epilogue = _SourceFile(Path("standalone_epilogue.h"))
-	epilogue.lines = [x for x in epilogue.path.open("rt") if x]
-	source_files.append(epilogue)
+	prologue = _SourceFile(src_dir / "standalone_prologue.h")
+	epilogue = _SourceFile(src_dir / "standalone_epilogue.h")
+	if not thin:
+		source_files.insert(0, prologue)
+		source_files.append(epilogue)
+		prologue.lines = [x for x in prologue.path.open("rt") if x]
+		epilogue.lines = [x for x in epilogue.path.open("rt") if x]
 
 	# Write the output file
 	_spam_header("Output")
@@ -127,6 +127,11 @@ def _main(src_dir, dest_dir):
 	with (dest_dir / "trace.h").open("wt") as out:
 		print("// Copyright Epic Games, Inc. All Rights Reserved.", file=out)
 		print("#pragma once", file=out)
+
+		if thin:
+			include_path = str(prologue.path.resolve()).replace("\\", "/")
+			print(f'#include "{include_path}"', file=out)
+
 		for source_file in source_files:
 			if not cpp_started:
 				if source_file.path.suffix == ".cpp":
@@ -141,14 +146,20 @@ def _main(src_dir, dest_dir):
 			for line in source_file.lines:
 				out.write(line)
 
+		if thin:
+			include_path = str(epilogue.path.resolve()).replace("\\", "/")
+			print("#endif // TRACE_IMPLEMENT", file=out)
+			print(f'#include "{include_path}"', file=out)
+
 	_spam("...done!")
 
 def main():
 	desc = "Amalgamate TraceLog into a standalone single-file library"
 	parser = argparse.ArgumentParser(description=desc)
 	parser.add_argument("outdir", help="Directory to write output file(s) to")
+	parser.add_argument("--thin", action="store_true", help="#include standalone_ instead of blitting")
 	args = parser.parse_args()
-	return _main(Path(__file__).parent, Path(args.outdir))
+	return _main(Path(__file__).parent, Path(args.outdir), args.thin)
 
 if __name__ == "__main__":
 	raise SystemExit(main())
