Conveniently copy LZ4 out if '--thin' is used

#rnx

diff --git a/Engine/Source/Runtime/TraceLog/create_standalone.py b/Engine/Source/Runtime/TraceLog/create_standalone.py
--- a/Engine/Source/Runtime/TraceLog/create_standalone.py
+++ b/Engine/Source/Runtime/TraceLog/create_standalone.py
@@ -1,5 +1,7 @@
 import os
 import re
+import stat
+import shutil
 import argparse
 from pathlib import Path
 
@@ -151,6 +153,18 @@ def _main(src_dir, dest_dir, thin):
 			print("#endif // TRACE_IMPLEMENT", file=out)
 			print(f'#include "{include_path}"', file=out)
 
+	# Copy some support files for thin iteration
+	if thin:
+		_spam_header("Support files for thin trace.h")
+		support_paths = (
+			src_dir / "Private/Trace/LZ4/lz4.h",
+			src_dir / "Private/Trace/LZ4/lz4.c.inl",
+		)
+		for support_path in support_paths:
+			dest_path = dest_dir / support_path.name
+			shutil.copy(support_path, dest_path)
+			os.chmod(dest_path, stat.S_IWRITE)
+
 	_spam("...done!")
 
 def main():
