Moved source file gathering into its own function
#rnx

diff --git a/Engine/Source/Runtime/TraceLog/create_standalone.py b/Engine/Source/Runtime/TraceLog/create_standalone.py
--- a/Engine/Source/Runtime/TraceLog/create_standalone.py
+++ b/Engine/Source/Runtime/TraceLog/create_standalone.py
@@ -39,14 +39,7 @@ def _exclude_line(line):
 	if not line:				return True
 
 #-------------------------------------------------------------------------------
-def _main(src_dir, dest_dir, thin):
-	dest_dir = dest_dir.resolve()
-	os.chdir(Path(__file__).parent)
-
-	_spam_header("Context")
-	_spam("Source dir:", os.getcwd())
-	_spam("Dest dir:", dest_dir)
-
+def _collect_source(src_dir):
 	_spam_header("Gathering source files")
 	files  = [x for x in src_dir.glob("Public/**/*")]
 	files += [x for x in src_dir.glob("Private/**/*")]
@@ -113,6 +106,19 @@ def _main(src_dir, dest_dir, thin):
 	# Exclude LZ4
 	source_files = [x for x in source_files if "lz4" not in x.path.name]
 
+	return source_files
+
+#-------------------------------------------------------------------------------
+def _main(src_dir, dest_dir, thin):
+	dest_dir = dest_dir.resolve()
+	src_dir = Path(__file__).parent
+
+	_spam_header("Trace")
+	_spam("Source dir:", src_dir)
+	_spam("Dest dir:", dest_dir)
+
+	source_files = _collect_source(src_dir)
+
 	# Add prologue and epilogue files
 	prologue = _SourceFile(src_dir / "standalone_prologue.h")
 	epilogue = _SourceFile(src_dir / "standalone_epilogue.h")
