Extra argument to allow filtering of the glob files
#rnx

diff --git a/Engine/Source/Runtime/TraceLog/create_standalone.py b/Engine/Source/Runtime/TraceLog/create_standalone.py
--- a/Engine/Source/Runtime/TraceLog/create_standalone.py
+++ b/Engine/Source/Runtime/TraceLog/create_standalone.py
@@ -39,13 +39,17 @@ def _exclude_line(line):
 	if not line:				return True
 
 #-------------------------------------------------------------------------------
-def _collect_source(src_dir):
+def _collect_source(src_dir, predicate=None):
 	_spam_header("Gathering source files")
 	files  = [x for x in src_dir.glob("Public/**/*")]
 	files += [x for x in src_dir.glob("Private/**/*")]
 	files  = [x for x in files if x.suffix in (".h", ".cpp", ".inl")]
 	_spam("Found %d files" % len(files))
 
+	if predicate:
+		files = [x for x in files if predicate(x)]
+		_spam("Filtered down to %d files" % len(files))
+
 	# Collect lines of each source file, filtering local includes
 	source_files = []
 	for file in files:
@@ -77,7 +81,8 @@ def _collect_source(src_dir):
 	for source_file in source_files:
 		new_deps = []
 		for dep in source_file.deps:
-			index = source_files.index(dep)
+			try: index = source_files.index(dep)
+			except ValueError: pass
 			dep = source_files[index]
 			new_deps.append(dep)
 		source_file.deps = new_deps
