Added a Vim tag
#rnx

diff --git a/Engine/Source/Runtime/TraceLog/create_standalone.py b/Engine/Source/Runtime/TraceLog/create_standalone.py
--- a/Engine/Source/Runtime/TraceLog/create_standalone.py
+++ b/Engine/Source/Runtime/TraceLog/create_standalone.py
@@ -188,3 +188,5 @@ def main():
 
 if __name__ == "__main__":
 	raise SystemExit(main())
+
+# vim: noexpandtab
