"enum used as boolean constant" compiler warning fix

#rnx

diff --git a/Engine/Source/Runtime/TraceLog/Public/Trace/Detail/Trace.h b/Engine/Source/Runtime/TraceLog/Public/Trace/Detail/Trace.h
--- a/Engine/Source/Runtime/TraceLog/Public/Trace/Detail/Trace.h
+++ b/Engine/Source/Runtime/TraceLog/Public/Trace/Detail/Trace.h
@@ -106,7 +106,7 @@ class FChannel;
 		EventProps_Meta const EventProps_Private = {}; \
 		typedef std::conditional<bIsImportant, UE::Trace::Private::FImportantLogScope, UE::Trace::Private::FLogScope>::type LogScopeType; \
 		explicit operator bool () const { return true; } \
-		enum { EventFlags = PartialEventFlags|(EventProps_Meta::NumAuxFields ? UE::Trace::Private::FEventInfo::Flag_MaybeHasAux : 0), }; \
+		enum { EventFlags = PartialEventFlags|((EventProps_Meta::NumAuxFields != 0) ? UE::Trace::Private::FEventInfo::Flag_MaybeHasAux : 0), }; \
 		static_assert( \
 			!bIsImportant || (uint32(EventFlags) & uint32(UE::Trace::Private::FEventInfo::Flag_NoSync)), \
 			"Trace events flagged as Important events must be marked NoSync" \
