Trace Snapshot to File

### Features
We can now write tailing memory to a utrace file. Snapshotting while tracing to a file or socket is supported.
This change introduces a new console command: Trace.SnapshotFile. It writes the profiling data from memory into an optionally specified a utrace file.

### Testing
Run Editor with -tracefile
Run Editor with -trace, connect live with insights
Run Editor without -trace (this does generate a very small utrace since there is memory and data in the tail even though default channels are not enabled)

diff --git a/Engine/Source/Runtime/Core/Private/ProfilingDebugging/TraceAuxiliary.cpp b/Engine/Source/Runtime/Core/Private/ProfilingDebugging/TraceAuxiliary.cpp
--- a/Engine/Source/Runtime/Core/Private/ProfilingDebugging/TraceAuxiliary.cpp
+++ b/Engine/Source/Runtime/Core/Private/ProfilingDebugging/TraceAuxiliary.cpp
@@ -91,6 +91,7 @@ public:
 	void					UpdateCsvStats() const;
 	void					StartWorkerThread();
 	void					StartEndFramePump();
+	bool					WriteSnapshot(const TCHAR* FilePath);
 
 private:
 	enum class EState : uint8
@@ -113,6 +114,7 @@ private:
 	void					DisableChannel(const TCHAR* Channel);
 	bool					SendToHost(const TCHAR* Host, const FTraceAuxiliary::FLogCategoryAlias& LogCategory);
 	bool					WriteToFile(const TCHAR* Path, const FTraceAuxiliary::FLogCategoryAlias& LogCategory);
+	bool					FinalizeFilePath(const TCHAR* InPath, FString& OutPath, const FTraceAuxiliary::FLogCategoryAlias& LogCategory);
 
 	typedef TMap<uint32, FChannelEntry, TInlineSetAllocator<128>> ChannelSet;
 	ChannelSet				CommandlineChannels;
@@ -407,7 +409,7 @@ bool FTraceAuxiliaryImpl::SendToHost(const TCHAR* Host, const FTraceAuxiliary::F
 }
 
 ////////////////////////////////////////////////////////////////////////////////
-bool FTraceAuxiliaryImpl::WriteToFile(const TCHAR* InPath, const FTraceAuxiliary::FLogCategoryAlias& LogCategory)
+bool FTraceAuxiliaryImpl::FinalizeFilePath(const TCHAR* InPath, FString& OutPath, const FTraceAuxiliary::FLogCategoryAlias& LogCategory)
 {
 	const FStringView Path(InPath);
 
@@ -417,7 +419,7 @@ bool FTraceAuxiliaryImpl::WriteToFile(const TCHAR* InPath, const FTraceAuxiliary
 	if (Path.IsEmpty())
 	{
 		const FString Name = GetDefaultName();
-		return WriteToFile(*Name, LogCategory);
+		return FinalizeFilePath(*Name, OutPath, LogCategory);
 	}
 
 	FString WritePath;
@@ -471,16 +473,28 @@ bool FTraceAuxiliaryImpl::WriteToFile(const TCHAR* InPath, const FTraceAuxiliary
 		return false;
 	}
 
-	// Finally, tell trace to write the trace to a file.
+	OutPath = MoveTemp(NativePath);
+	return true;
+}
+
+////////////////////////////////////////////////////////////////////////////////
+bool FTraceAuxiliaryImpl::WriteToFile(const TCHAR* Path, const FTraceAuxiliary::FLogCategoryAlias& LogCategory)
+{
+	FString NativePath;
+	if (!FinalizeFilePath(Path, NativePath, LogCategory))
+	{
+		return false;
+	}
+
 	if (!UE::Trace::WriteTo(*NativePath))
 	{
-		if (FPathViews::Equals(NativePath, WritePath))
+		if (FPathViews::Equals(NativePath, FStringView(Path)))
 		{
 			UE_LOG_REF(LogCategory, Warning, TEXT("Unable to trace to file '%s'"), *NativePath);
 		}
 		else
 		{
-			UE_LOG_REF(LogCategory, Warning, TEXT("Unable to trace to file '%s' (transformed from '%s')"), *NativePath, *WritePath);
+			UE_LOG_REF(LogCategory, Warning, TEXT("Unable to trace to file '%s' (transformed from '%s')"), *NativePath, Path ? Path : TEXT("null"));
 		}
 		return false;
 	}
@@ -489,6 +503,18 @@ bool FTraceAuxiliaryImpl::WriteToFile(const TCHAR* InPath, const FTraceAuxiliary
 	return true;
 }
 
+////////////////////////////////////////////////////////////////////////////////
+bool FTraceAuxiliaryImpl::WriteSnapshot(const TCHAR* FilePath)
+{
+	FString NativePath;
+	if (!FinalizeFilePath(FilePath, NativePath, LogConsoleResponse))
+	{
+		return false;
+	}
+
+	return UE::Trace::WriteSnapshotTo(*NativePath);
+}
+
 ////////////////////////////////////////////////////////////////////////////////
 const TCHAR* FTraceAuxiliaryImpl::GetDest() const
 {
@@ -788,6 +814,32 @@ static void TraceAuxiliaryDisableChannels(const TArray<FString>& Args)
 	UE_LOG(LogConsoleResponse, Display, TEXT("Enabled channels: %s"), EnabledChannels.ToString());
 }
 
+////////////////////////////////////////////////////////////////////////////////
+static void TraceAuxiliarySnapshotFile(const TArray<FString>& Args)
+{
+	const TCHAR* FilePath = nullptr;
+
+	if (Args.Num() == 1)
+	{
+		FilePath = *Args[0];
+	}
+	else if (Args.Num() > 1)
+	{
+		UE_LOG(LogConsoleResponse, Warning, TEXT("Invalid arguments. Usage: Trace.SnapshotFile [Path]"));
+		return;
+	}
+
+	double StartTime = FPlatformTime::Seconds();
+	if (GTraceAuxiliary.WriteSnapshot(FilePath))
+	{
+		UE_LOG(LogConsoleResponse, Display, TEXT("Trace snapshot generated in %f seconds."), FPlatformTime::Seconds() - StartTime);
+	}
+	else
+	{
+		UE_LOG(LogConsoleResponse, Error, TEXT("Trace snapshot failed."));
+	}
+}
+
 ////////////////////////////////////////////////////////////////////////////////
 static FAutoConsoleCommand TraceAuxiliarySendCmd(
 	TEXT("Trace.Send"),
@@ -864,6 +916,14 @@ static FAutoConsoleCommand TraceAuxiliaryChannelDisableCmd(
 	FConsoleCommandWithArgsDelegate::CreateStatic(TraceAuxiliaryDisableChannels)
 );
 
+////////////////////////////////////////////////////////////////////////////////
+static FAutoConsoleCommand TraceAuxiliarySnapshotFileCmd(
+	TEXT("Trace.SnapshotFile"),
+	TEXT("[Path] - Writes a snapshot of the current in-memory trace buffer to a file."
+	),
+	FConsoleCommandWithArgsDelegate::CreateStatic(TraceAuxiliarySnapshotFile)
+);
+
 #endif // UE_TRACE_ENABLED
 
 
@@ -1157,6 +1217,15 @@ bool FTraceAuxiliary::Resume()
 	return true;
 }
 
+////////////////////////////////////////////////////////////////////////////////
+bool FTraceAuxiliary::WriteSnapshot(const TCHAR* FilePath)
+{
+#if UE_TRACE_ENABLED
+	return GTraceAuxiliary.WriteSnapshot(FilePath);
+#endif
+	return true;
+}
+
 ////////////////////////////////////////////////////////////////////////////////
 void FTraceAuxiliary::Initialize(const TCHAR* CommandLine)
 {
diff --git a/Engine/Source/Runtime/Core/Public/ProfilingDebugging/TraceAuxiliary.h b/Engine/Source/Runtime/Core/Public/ProfilingDebugging/TraceAuxiliary.h
--- a/Engine/Source/Runtime/Core/Public/ProfilingDebugging/TraceAuxiliary.h
+++ b/Engine/Source/Runtime/Core/Public/ProfilingDebugging/TraceAuxiliary.h
@@ -71,6 +71,12 @@ public:
 	 */
 	static bool Resume();
 
+	/**
+	 * Write tailing memory state to a utrace file.
+	 * @param FilePath Path to the file to write the snapshot to. If it is null or empty a file path will be generated.
+	 */
+	static bool WriteSnapshot(const TCHAR* FilePath);
+
 	/**
 	 * Initialize Trace systems.
 	 * @param CommandLine to use for initializing
diff --git a/Engine/Source/Runtime/TraceLog/Private/Trace/EventNode.cpp b/Engine/Source/Runtime/TraceLog/Private/Trace/EventNode.cpp
--- a/Engine/Source/Runtime/TraceLog/Private/Trace/EventNode.cpp
+++ b/Engine/Source/Runtime/TraceLog/Private/Trace/EventNode.cpp
@@ -35,7 +35,16 @@ const FEventNode* FEventNode::FIter::GetNext()
 	return Ret;
 }
 
+////////////////////////////////////////////////////////////////////////////////
+FEventNode::FIter FEventNode::Read()
+{
+	if (GEventListHead)
+	{
+		return { GEventListHead };
+	}
 
+	return {};
+}
 
 ////////////////////////////////////////////////////////////////////////////////
 FEventNode::FIter FEventNode::ReadNew()
diff --git a/Engine/Source/Runtime/TraceLog/Private/Trace/Tail.cpp b/Engine/Source/Runtime/TraceLog/Private/Trace/Tail.cpp
--- a/Engine/Source/Runtime/TraceLog/Private/Trace/Tail.cpp
+++ b/Engine/Source/Runtime/TraceLog/Private/Trace/Tail.cpp
@@ -286,6 +286,12 @@ void Writer_InitializeTail(int32 BufferSize)
 	GPacketRing.Initialize(BufferSize);
 }
 
+////////////////////////////////////////////////////////////////////////////////
+bool Writer_IsTailing()
+{
+	return GPacketRing.IsActive();
+}
+
 ////////////////////////////////////////////////////////////////////////////////
 void Writer_ShutdownTail()
 {
diff --git a/Engine/Source/Runtime/TraceLog/Private/Trace/Trace.cpp b/Engine/Source/Runtime/TraceLog/Private/Trace/Trace.cpp
--- a/Engine/Source/Runtime/TraceLog/Private/Trace/Trace.cpp
+++ b/Engine/Source/Runtime/TraceLog/Private/Trace/Trace.cpp
@@ -22,6 +22,7 @@ void	Writer_Shutdown();
 void	Writer_Update();
 bool	Writer_SendTo(const ANSICHAR*, uint32);
 bool	Writer_WriteTo(const ANSICHAR*);
+bool	Writer_WriteSnapshotTo(const ANSICHAR*);
 bool	Writer_IsTracing();
 bool	Writer_Stop();
 uint32	Writer_GetThreadId();
@@ -96,6 +97,14 @@ bool WriteTo(const TCHAR* InPath)
 	return Private::Writer_WriteTo(Path);
 }
 
+////////////////////////////////////////////////////////////////////////////////
+bool WriteSnapshotTo(const TCHAR* InPath)
+{
+	char Path[512];
+	ToAnsiCheap(Path, InPath);
+	return Private::Writer_WriteSnapshotTo(Path);
+}
+
 ////////////////////////////////////////////////////////////////////////////////
 bool IsTracing()
 {
diff --git a/Engine/Source/Runtime/TraceLog/Private/Trace/Writer.cpp b/Engine/Source/Runtime/TraceLog/Private/Trace/Writer.cpp
--- a/Engine/Source/Runtime/TraceLog/Private/Trace/Writer.cpp
+++ b/Engine/Source/Runtime/TraceLog/Private/Trace/Writer.cpp
@@ -56,6 +56,8 @@ void			Writer_UpdateControl();
 void			Writer_InitializeControl();
 void			Writer_ShutdownControl();
 bool			Writer_IsTracing();
+bool			Writer_IsTailing();
+static bool		Writer_SessionPrologue();
 
 
 
@@ -340,11 +342,10 @@ void Writer_SendData(uint32 ThreadId, uint8* __restrict Data, uint32 Size)
 }
 
 ////////////////////////////////////////////////////////////////////////////////
-static void Writer_DescribeEvents()
+static void Writer_DescribeEvents(FEventNode::FIter Iter)
 {
 	TWriteBufferRedirect<4096> TraceData;
 
-	FEventNode::FIter Iter = FEventNode::ReadNew();
 	while (const FEventNode* Event = Iter.GetNext())
 	{
 		Event->Describe();
@@ -382,7 +383,7 @@ static void Writer_DescribeAnnounce()
 	}
 
 	Writer_AnnounceChannels();
-	Writer_DescribeEvents();
+	Writer_DescribeEvents(FEventNode::ReadNew());
 }
 
 
@@ -403,7 +404,7 @@ static void Writer_SendSync()
 	// update that are newer than events sent it the following update where IO
 	// is established. This will result in holes in serial numbering. A few sync
 	// points are sent to aid analysis in determining what are holes and what is
-	// just a requirement for more data. Holws will only occurr at the start.
+	// just a requirement for more data. Holes will only occur at the start.
 
 	// Note that Sync is alias as Important/Internal as changing Bias would
 	// break backwards compatibility.
@@ -414,6 +415,18 @@ static void Writer_SendSync()
 	--GSyncPacketCountdown;
 }
 
+////////////////////////////////////////////////////////////////////////////////
+static void Writer_Close()
+{
+	if (GDataHandle)
+	{
+		Writer_FlushSendBuffer();
+		IoClose(GDataHandle);
+	}
+
+	GDataHandle = 0;
+}
+
 ////////////////////////////////////////////////////////////////////////////////
 static bool Writer_UpdateConnection()
 {
@@ -431,13 +444,7 @@ static bool Writer_UpdateConnection()
 
 		if (GPendingDataHandle == (~0ull -CloseInertia))
 		{
-			if (GDataHandle)
-			{
-				Writer_FlushSendBuffer();
-				IoClose(GDataHandle);
-			}
-
-			GDataHandle = 0;
+			Writer_Close();
 			GPendingDataHandle = 0;
 		}
 
@@ -454,6 +461,36 @@ static bool Writer_UpdateConnection()
 
 	GDataHandle = GPendingDataHandle;
 	GPendingDataHandle = 0;
+	if (!Writer_SessionPrologue())
+	{
+		return false;
+	}
+
+	// Reset statistics.
+	GTraceStatistics.BytesSent = 0;
+	GTraceStatistics.BytesTraced = 0;
+
+	// The first events we will send are ones that describe the trace's events
+	FEventNode::OnConnect();
+	Writer_DescribeEvents(FEventNode::ReadNew());
+
+	// Send cached events (i.e. importants) and the tail of recent events
+	Writer_CacheOnConnect();
+	Writer_TailOnConnect();
+
+	// See Writer_SendSync() for details.
+	GSyncPacketCountdown = GNumSyncPackets;
+
+	return true;
+}
+
+////////////////////////////////////////////////////////////////////////////////
+static bool Writer_SessionPrologue()
+{
+	if (!GDataHandle)
+	{
+		return false;
+	}
 
 #if TRACE_PRIVATE_BUFFER_SEND
 	if (!GSendBuffer)
@@ -493,21 +530,6 @@ static bool Writer_UpdateConnection()
 		return false;
 	}
 
-	// Reset statistics.
-	GTraceStatistics.BytesSent = 0;
-	GTraceStatistics.BytesTraced = 0;
-
-	// The first events we will send are ones that describe the trace's events
-	FEventNode::OnConnect();
-	Writer_DescribeEvents();
-
-	// Send cached events (i.e. importants) and the tail of recent events
-	Writer_CacheOnConnect();
-	Writer_TailOnConnect();
-
-	// See Writer_SendSync() for details.
-	GSyncPacketCountdown = GNumSyncPackets;
-
 	return true;
 }
 
@@ -519,13 +541,8 @@ static volatile bool	GWorkerThreadQuit;	// = false;
 static volatile unsigned int	GUpdateInProgress = 1;	// Don't allow updates until initalized
 
 ////////////////////////////////////////////////////////////////////////////////
-static void Writer_WorkerUpdate()
+static void Writer_WorkerUpdateInternal()
 {
-	if (!AtomicCompareExchangeAcquire(&GUpdateInProgress, 1u, 0u))
-	{
-		return;
-	}
-	
 	Writer_UpdateControl();
 	Writer_UpdateConnection();
 	Writer_DescribeAnnounce();
@@ -540,6 +557,17 @@ static void Writer_WorkerUpdate()
 		Writer_FlushSendBuffer();
 	}
 #endif
+}
+
+////////////////////////////////////////////////////////////////////////////////
+static void Writer_WorkerUpdate()
+{
+	if (!AtomicCompareExchangeAcquire(&GUpdateInProgress, 1u, 0u))
+	{
+		return;
+	}
+	
+	Writer_WorkerUpdateInternal();
 
 	AtomicExchangeRelease(&GUpdateInProgress, 0u);
 }
@@ -749,6 +777,133 @@ bool Writer_WriteTo(const ANSICHAR* Path)
 	return true;
 }
 
+////////////////////////////////////////////////////////////////////////////////
+struct WorkerUpdateLock
+{
+	WorkerUpdateLock()
+	{
+		CyclesPerSecond = TimeGetFrequency();
+		StartSeconds = GetTime();
+
+		while (!AtomicCompareExchangeAcquire(&GUpdateInProgress, 1u, 0u))
+		{
+			ThreadSleep(0);
+
+			if (TimedOut())
+			{
+				break;
+			}
+		}
+	}
+
+	~WorkerUpdateLock()
+	{
+		AtomicExchangeRelease(&GUpdateInProgress, 0u);
+	}
+
+	double GetTime()
+	{
+		return static_cast<double>(TimeGetTimestamp()) / static_cast<double>(CyclesPerSecond);
+	}
+
+	bool TimedOut()
+	{
+		uint64 WaitTime = GetTime() - StartSeconds;
+		return WaitTime > MaxWaitSeconds;
+	}
+
+	uint64 CyclesPerSecond;
+	double StartSeconds;
+	inline const static double MaxWaitSeconds = 1.0;
+};
+
+////////////////////////////////////////////////////////////////////////////////
+template <typename Type>
+struct TStashGlobal
+{
+	TStashGlobal(Type& Global)
+		: Variable(Global)
+		, Stashed(Global)
+	{
+		Variable = {};
+	}
+
+	TStashGlobal(Type& Global, const Type& Value)
+		: Variable(Global)
+		, Stashed(Global)
+	{
+		Variable = Value;
+	}
+
+	~TStashGlobal()
+	{
+		Variable = Stashed;
+	}
+
+private:
+	Type& Variable;
+	Type Stashed;
+};
+
+////////////////////////////////////////////////////////////////////////////////
+bool Writer_WriteSnapshotTo(const ANSICHAR* Path)
+{
+	if (!Writer_IsTailing())
+	{
+		return false;
+	}
+
+	WorkerUpdateLock UpdateLock;
+
+	// We have a timeout just in case the worker thread goes off the rails
+	//  We are called by diagnostic handlers like crash reporter, do not deadlock
+	if (UpdateLock.TimedOut())
+	{
+		return false;
+	}
+	
+	// Bring everything up to date with the active tracing connection.
+	//  Any connection writes after we call the worker update
+	//  will need to treat source data structures as read-only.
+	//  Those structures can be stateful about what has or has not been
+	//  written (cursor state, "new" event, etc) to the tracing connection.
+	//  We are pre-empting the connection to write the snapshot.
+	//
+	// Doing a full pump of the data here is more robust than opening
+	//  pathways through each data structure for immutable writes because
+	//  the data structures are order-dependent and inter-referential.
+	//  Not writing all state can very easily lead to bugs in either the
+	//  snapshot or, even worse, the tracing connection. Parsers only
+	//  have a limited tolerance to gaps/out-of-order event packets.
+	Writer_WorkerUpdateInternal();
+
+	{
+		TStashGlobal DataHandle(GDataHandle);
+		TStashGlobal SyncPacketCountdown(GSyncPacketCountdown, GNumSyncPackets);
+		TStashGlobal TraceStatistics(GTraceStatistics);
+
+		// Open the snapshot file and write the file header
+		GDataHandle = FileOpen(Path);
+		if (!GDataHandle || !Writer_SessionPrologue())
+		{
+			return false;
+		}
+
+		// The first events we will send are ones that describe the trace's events
+		Writer_DescribeEvents(FEventNode::Read());
+
+		// Send cached events (i.e. importants) and the tail of recent events
+		Writer_CacheOnConnect();
+		Writer_TailOnConnect();
+
+		// Send sync packets to help parsers digest any out-of-order events
+		Writer_SendSync();
+		Writer_Close();
+	}
+
+	return true;
+}
+
 ////////////////////////////////////////////////////////////////////////////////
 bool Writer_IsTracing()
 {
diff --git a/Engine/Source/Runtime/TraceLog/Public/Trace/Detail/EventNode.h b/Engine/Source/Runtime/TraceLog/Public/Trace/Detail/EventNode.h
--- a/Engine/Source/Runtime/TraceLog/Public/Trace/Detail/EventNode.h
+++ b/Engine/Source/Runtime/TraceLog/Public/Trace/Detail/EventNode.h
@@ -46,6 +46,7 @@ public:
 		void*				Inner;
 	};
 
+	static FIter			Read();
 	static FIter			ReadNew();
 	static void				OnConnect();
 	TRACELOG_API uint32		Initialize(const FEventInfo* InInfo);
diff --git a/Engine/Source/Runtime/TraceLog/Public/Trace/Trace.h b/Engine/Source/Runtime/TraceLog/Public/Trace/Trace.h
--- a/Engine/Source/Runtime/TraceLog/Public/Trace/Trace.h
+++ b/Engine/Source/Runtime/TraceLog/Public/Trace/Trace.h
@@ -105,6 +105,7 @@ UE_TRACE_API void	Update() UE_TRACE_IMPL();
 UE_TRACE_API void	GetStatistics(FStatistics& Out) UE_TRACE_IMPL();
 UE_TRACE_API bool	SendTo(const TCHAR* Host, uint32 Port=0) UE_TRACE_IMPL(false);
 UE_TRACE_API bool	WriteTo(const TCHAR* Path) UE_TRACE_IMPL(false);
+UE_TRACE_API bool	WriteSnapshotTo(const TCHAR* Path) UE_TRACE_IMPL(false);
 UE_TRACE_API bool	IsTracing() UE_TRACE_IMPL(false);
 UE_TRACE_API bool	Stop() UE_TRACE_IMPL(false);
 UE_TRACE_API bool	IsChannel(const TCHAR* ChanneName) UE_TRACE_IMPL(false);
