[TraceLog] Fix for static analysis warnings

diff --git a/Engine/Source/Runtime/TraceLog/Public/Trace/Detail/Trace.h b/Engine/Source/Runtime/TraceLog/Public/Trace/Detail/Trace.h
--- a/Engine/Source/Runtime/TraceLog/Public/Trace/Detail/Trace.h
+++ b/Engine/Source/Runtime/TraceLog/Public/Trace/Detail/Trace.h
@@ -124,9 +124,9 @@ class FChannel;
 			FieldName##_Meta::Offset + FieldName##_Meta::Size,
 
 #define TRACE_PRIVATE_EVENT_END() \
-		std::conditional<bIsDefinition, DefinitionType, UE::Trace::DisabledField>::type> DefinitionId_Meta;\
+		std::conditional<bIsDefinition != 0, DefinitionType, UE::Trace::DisabledField>::type> DefinitionId_Meta;\
 		DefinitionId_Meta const DefinitionId_Field = UE::Trace::FLiteralName("");\
-		static constexpr uint16 NumDefinitionFields = bIsDefinition ? 1 : 0;\
+		static constexpr uint16 NumDefinitionFields = (bIsDefinition != 0) ? 1 : 0;\
 		template<typename RefType>\
 		auto SetDefinitionId(const RefType& Id) const \
 		{ \
@@ -135,11 +135,11 @@ class FChannel;
 		} \
 		typedef UE::Trace::TField<DefinitionId_Meta::Index + NumDefinitionFields, DefinitionId_Meta::Offset + DefinitionId_Meta::Size, UE::Trace::EventProps> EventProps_Meta; \
 		EventProps_Meta const EventProps_Private = {}; \
-		typedef std::conditional<bIsImportant, UE::Trace::Private::FImportantLogScope, UE::Trace::Private::FLogScope>::type LogScopeType; \
+		typedef std::conditional<bIsImportant != 0, UE::Trace::Private::FImportantLogScope, UE::Trace::Private::FLogScope>::type LogScopeType; \
 		explicit operator bool () const { return true; } \
 		enum { EventFlags = PartialEventFlags|((EventProps_Meta::NumAuxFields != 0) ? UE::Trace::Private::FEventInfo::Flag_MaybeHasAux : 0), }; \
 		static_assert( \
-			!bIsImportant || (uint32(EventFlags) & uint32(UE::Trace::Private::FEventInfo::Flag_NoSync)), \
+			(bIsImportant == 0) || (uint32(EventFlags) & uint32(UE::Trace::Private::FEventInfo::Flag_NoSync)), \
 			"Trace events flagged as Important events must be marked NoSync" \
 		); \
 	};
