Trace: Fix bug where Snapshots don't contain event descriptions if no active connection exists

#rnx

diff --git a/Engine/Source/Runtime/TraceLog/Private/Trace/EventNode.cpp b/Engine/Source/Runtime/TraceLog/Private/Trace/EventNode.cpp
--- a/Engine/Source/Runtime/TraceLog/Private/Trace/EventNode.cpp
+++ b/Engine/Source/Runtime/TraceLog/Private/Trace/EventNode.cpp
@@ -38,11 +38,18 @@ const FEventNode* FEventNode::FIter::GetNext()
 ////////////////////////////////////////////////////////////////////////////////
 FEventNode::FIter FEventNode::Read()
 {
+	// This is the current event list as of the last ReadNew() call
 	if (GEventListHead)
 	{
 		return { GEventListHead };
 	}
 
+	// This is the current event list if ReadNew() has never been called
+	if (GNewEventList)
+	{
+		return { GNewEventList };
+	}
+
 	return {};
 }
 
