TraceLog HoloLens platform extension file.

#rnx

diff --git a/Engine/Platforms/Hololens/Source/Runtime/TraceLog/TraceLog_HoloLens.build.cs b/Engine/Platforms/Hololens/Source/Runtime/TraceLog/TraceLog_HoloLens.build.cs
--- a/Engine/Platforms/Hololens/Source/Runtime/TraceLog/TraceLog_HoloLens.build.cs
+++ b/Engine/Platforms/Hololens/Source/Runtime/TraceLog/TraceLog_HoloLens.build.cs
@@ -6,6 +6,7 @@ namespace UnrealBuildTool.Rules
 	{
 		public TraceLog_HoloLens(ReadOnlyTargetRules Target) : base(Target)
 		{
+			EnableTraceByDefault(Target);
 		}
 	}
 }
diff --git a/Engine/Source/Runtime/TraceLog/Public/Trace/Config.h b/Engine/Source/Runtime/TraceLog/Public/Trace/Config.h
--- a/Engine/Source/Runtime/TraceLog/Public/Trace/Config.h
+++ b/Engine/Source/Runtime/TraceLog/Public/Trace/Config.h
@@ -6,7 +6,7 @@
 
 #if !defined(UE_TRACE_ENABLED)
 #	if !UE_BUILD_SHIPPING && !IS_PROGRAM
-#		if PLATFORM_WINDOWS || PLATFORM_UNIX || PLATFORM_APPLE || PLATFORM_ANDROID || PLATFORM_HOLOLENS
+#		if PLATFORM_WINDOWS || PLATFORM_UNIX || PLATFORM_APPLE || PLATFORM_ANDROID
 #			define UE_TRACE_ENABLED	1
 #		endif
 #	endif
