Fix for CIS errors on latest MSVC toolchain

diff --git a/Engine/Source/Runtime/TraceLog/Public/Trace/Detail/Important/ImportantLogScope.inl b/Engine/Source/Runtime/TraceLog/Public/Trace/Detail/Important/ImportantLogScope.inl
--- a/Engine/Source/Runtime/TraceLog/Public/Trace/Detail/Important/ImportantLogScope.inl
+++ b/Engine/Source/Runtime/TraceLog/Public/Trace/Detail/Important/ImportantLogScope.inl
@@ -27,7 +27,7 @@ TRACELOG_API FNextSharedBuffer				Writer_NextSharedBuffer(FSharedBuffer*, int32,
 template <class T>
 FORCENOINLINE FImportantLogScope FImportantLogScope::Enter(uint32 ArrayDataSize)
 {
-	static_assert(uint32(T::EventFlags) & uint32(FEventInfo::Flag_MaybeHasAux), "Only important trace events with array-type fields need a size parameter to UE_TRACE_LOG()");
+	static_assert(!!(uint32(T::EventFlags) & uint32(FEventInfo::Flag_MaybeHasAux)), "Only important trace events with array-type fields need a size parameter to UE_TRACE_LOG()");
 
 	ArrayDataSize += sizeof(FAuxHeader) * T::EventProps_Meta::NumAuxFields;
 	ArrayDataSize += 1; // for AuxDataTerminal
