[TraceLog] Add OnConnection callback

Adds a callback which is issued from Trace whenever a new connection is established. Any event emitted inside the callback are collected separately and written to the trace stream after important events, but before events collected in the tail buffer. In the engine layer the callback is exposed as static delegate on TraceAuxiliary.

diff --git a/Engine/Source/Runtime/Core/Private/ProfilingDebugging/TraceAuxiliary.cpp b/Engine/Source/Runtime/Core/Private/ProfilingDebugging/TraceAuxiliary.cpp
--- a/Engine/Source/Runtime/Core/Private/ProfilingDebugging/TraceAuxiliary.cpp
+++ b/Engine/Source/Runtime/Core/Private/ProfilingDebugging/TraceAuxiliary.cpp
@@ -79,7 +79,7 @@ public:
 	void					GetActiveChannelsString(FStringBuilderBase& String) const;
 	void					AddCommandlineChannels(const TCHAR* ChannelList);
 	void					ResetCommandlineChannels();
-	bool					HasCommandlineChannels() const { return !CommandlineChannels.IsEmpty(); };
+	bool					HasCommandlineChannels() const { return !CommandlineChannels.IsEmpty(); }
 	void					EnableChannels(const TCHAR* ChannelList);
 	void					DisableChannels(const TCHAR* ChannelList);
 	bool					Connect(ETraceConnectType Type, const TCHAR* Parameter, const FTraceAuxiliary::FLogCategoryAlias& LogCategory);
@@ -596,6 +596,11 @@ void FTraceAuxiliaryImpl::StartEndFramePump()
 	}
 }
 
+////////////////////////////////////////////////////////////////////////////////
+void OnConnectionCallback()
+{
+	FTraceAuxiliary::OnConnection.Broadcast();
+}
 
 ////////////////////////////////////////////////////////////////////////////////
 static void TraceAuxiliaryConnectEpilogue()
@@ -1152,6 +1157,9 @@ static bool StartFromCommandlineArguments(const TCHAR* CommandLine)
 	return false;
 }
 
+////////////////////////////////////////////////////////////////////////////////
+FTraceAuxiliary::FOnConnection FTraceAuxiliary::OnConnection;
+
 ////////////////////////////////////////////////////////////////////////////////
 bool FTraceAuxiliary::Start(EConnectionType Type, const TCHAR* Target, const TCHAR* Channels, Options* Options, const FLogCategoryAlias& LogCategory)
 {
@@ -1296,6 +1304,7 @@ void FTraceAuxiliary::Initialize(const TCHAR* CommandLine)
 	UE::Trace::FInitializeDesc Desc;
 	Desc.bUseWorkerThread = false;
 	Desc.bUseImportantCache = (FParse::Param(CommandLine, TEXT("tracenocache")) == false);
+	Desc.OnConnectionFunc = &OnConnectionCallback; 
 	if (FParse::Value(CommandLine, TEXT("-tracetailmb="), Desc.TailSizeBytes))
 	{
 		Desc.TailSizeBytes <<= 20;
diff --git a/Engine/Source/Runtime/Core/Public/ProfilingDebugging/TraceAuxiliary.h b/Engine/Source/Runtime/Core/Public/ProfilingDebugging/TraceAuxiliary.h
--- a/Engine/Source/Runtime/Core/Public/ProfilingDebugging/TraceAuxiliary.h
+++ b/Engine/Source/Runtime/Core/Public/ProfilingDebugging/TraceAuxiliary.h
@@ -18,6 +18,11 @@ public:
 	typedef FLogCategoryBase FLogCategoryAlias;
 #endif
 
+	/**
+	 * Callback type when a new connection is established.
+	 */
+	DECLARE_MULTICAST_DELEGATE(FOnConnection);
+
 	enum class EConnectionType
 	{
 		/**
@@ -123,5 +128,18 @@ public:
 	static void	GetActiveChannelsString(FStringBuilderBase& String);
 
 
+	/**
+	 * Delegate that triggers when a connection is established. Gives subscribers a chance to trace events that appear
+	 * after important events but before regular events (including tail). The following restrictions apply:
+	 *  * Only NoSync event types can be emitted.
+	 *  * Important events should not be emitted. They will appear after the events in the tail.
+	 *  * Callback is issued from a worker thread. User is responsible to synchronize shared resources.
+	 * 
+	 * @note This is an advanced feature to avoid using important events in cases where event data can be
+	 *		 recalled easily.
+	 *		 
+	 * @param Callback Delegate to call on new connections.
+	 */
+	static FOnConnection OnConnection;
 	
 };
diff --git a/Engine/Source/Runtime/TraceLog/Private/Trace/Tail.cpp b/Engine/Source/Runtime/TraceLog/Private/Trace/Tail.cpp
--- a/Engine/Source/Runtime/TraceLog/Private/Trace/Tail.cpp
+++ b/Engine/Source/Runtime/TraceLog/Private/Trace/Tail.cpp
@@ -40,6 +40,7 @@ public:
 	void			Initialize(uint32 InSize);
 	void			Shutdown();
 	void			Reset();
+	void			Pause(bool bInPause) { bPaused = bInPause; };
 	uint32			GetSize() const;
 	bool			IsActive() const;
 	FRange			GetBackPackets() const;
@@ -57,6 +58,7 @@ private:
 	uint32			Cursor;
 	uint32			Left;
 	uint32			Right;
+	bool			bPaused;
 };
 
 // TraceLog must be ready after value-init so that it can be used before
@@ -68,6 +70,7 @@ void FPacketRing::Initialize(uint32 InSize)
 {
 	Data = (uint8*)Writer_MemoryAllocate(InSize, 16);
 	Size = InSize;
+	bPaused = false;
 	Reset();
 }
 
@@ -94,7 +97,7 @@ uint32 FPacketRing::GetSize() const
 ////////////////////////////////////////////////////////////////////////////////
 bool FPacketRing::IsActive() const
 {
-	return Data != nullptr;
+	return Data != nullptr && !bPaused;
 }
 
 ////////////////////////////////////////////////////////////////////////////////
@@ -199,6 +202,11 @@ FTidPacketBase* FPacketRing::AppendImpl(uint32 InSize)
 static FPacketRing GPacketRing; // = {};
 
 
+////////////////////////////////////////////////////////////////////////////////
+void Writer_TailPause(bool bPause)
+{
+	GPacketRing.Pause(bPause);
+}
 
 ////////////////////////////////////////////////////////////////////////////////
 void Writer_TailAppend(uint32 ThreadId, uint8* __restrict Data, uint32 Size, bool bPartial)
diff --git a/Engine/Source/Runtime/TraceLog/Private/Trace/TlsBuffer.cpp b/Engine/Source/Runtime/TraceLog/Private/Trace/TlsBuffer.cpp
--- a/Engine/Source/Runtime/TraceLog/Private/Trace/TlsBuffer.cpp
+++ b/Engine/Source/Runtime/TraceLog/Private/Trace/TlsBuffer.cpp
@@ -218,6 +218,12 @@ void Writer_DrainBuffers()
 	}
 }
 
+////////////////////////////////////////////////////////////////////////////////
+bool Writer_DrainLocalBuffer(uint32 ThreadId)
+{
+	return Writer_DrainBuffer(ThreadId, GTlsWriteBuffer);
+}
+	
 ////////////////////////////////////////////////////////////////////////////////
 void Writer_EndThreadBuffer()
 {
diff --git a/Engine/Source/Runtime/TraceLog/Private/Trace/Writer.cpp b/Engine/Source/Runtime/TraceLog/Private/Trace/Writer.cpp
--- a/Engine/Source/Runtime/TraceLog/Private/Trace/Writer.cpp
+++ b/Engine/Source/Runtime/TraceLog/Private/Trace/Writer.cpp
@@ -41,14 +41,17 @@ void			Writer_SendData(uint32, uint8* __restrict, uint32);
 void			Writer_InitializeTail(int32);
 void			Writer_ShutdownTail();
 void			Writer_TailOnConnect();
+void			Writer_TailPause(bool);	
 void			Writer_InitializeSharedBuffers();
 void			Writer_ShutdownSharedBuffers();
 void			Writer_UpdateSharedBuffers();
 void			Writer_InitializeCache();
 void			Writer_ShutdownCache();
 void			Writer_CacheOnConnect();
+void			Writer_CallbackOnConnect();
 void			Writer_InitializePool();
 void			Writer_ShutdownPool();
+bool			Writer_DrainLocalBuffer(uint32);
 void			Writer_DrainBuffers();
 void			Writer_EndThreadBuffer();
 uint32			Writer_GetControlPort();
@@ -58,6 +61,7 @@ void			Writer_ShutdownControl();
 bool			Writer_IsTracing();
 bool			Writer_IsTailing();
 static bool		Writer_SessionPrologue();
+void			Writer_FreeBlockListToPool(FWriteBuffer*, FWriteBuffer*);
 
 
 
@@ -391,6 +395,7 @@ static void Writer_DescribeAnnounce()
 ////////////////////////////////////////////////////////////////////////////////
 static int8			GSyncPacketCountdown;	// = 0
 static const int8	GNumSyncPackets			= 3;
+static OnConnectFunc*	GOnConnection = nullptr;
 
 ////////////////////////////////////////////////////////////////////////////////
 static void Writer_SendSync()
@@ -474,8 +479,12 @@ static bool Writer_UpdateConnection()
 	FEventNode::OnConnect();
 	Writer_DescribeEvents(FEventNode::ReadNew());
 
-	// Send cached events (i.e. importants) and the tail of recent events
+	// Send cached events (i.e. importants) 
 	Writer_CacheOnConnect();
+	// Issue on connection callback. This allows writing events that are
+	// not cached but important for the cache
+	Writer_CallbackOnConnect();
+	// Finally write the events in the tail buffer
 	Writer_TailOnConnect();
 
 	// See Writer_SendSync() for details.
@@ -533,12 +542,28 @@ static bool Writer_SessionPrologue()
 	return true;
 }
 
+////////////////////////////////////////////////////////////////////////////////
+void Writer_CallbackOnConnect()
+{
+	if (!GOnConnection)
+	{
+		return;
+	}
+	
+	Writer_TailPause(true);
+	// Issue callback. We assume any events emitted here are not marked as
+	// important and emitted on this thread.
+	GOnConnection();
+	// Drain only the events for this thread passing tail as you go.
+	Writer_DrainLocalBuffer(ETransportTid::Bias);
+	Writer_TailPause(false);
+}
 
 
 ////////////////////////////////////////////////////////////////////////////////
 static UPTRINT			GWorkerThread;		// = 0;
 static volatile bool	GWorkerThreadQuit;	// = false;
-static volatile unsigned int	GUpdateInProgress = 1;	// Don't allow updates until initalized
+static volatile unsigned int	GUpdateInProgress = 1;	// Don't allow updates until initialized
 
 ////////////////////////////////////////////////////////////////////////////////
 static void Writer_WorkerUpdateInternal()
@@ -714,6 +739,9 @@ void Writer_Initialize(const FInitializeDesc& Desc)
 		Writer_WorkerCreate();
 	}
 
+	// Store callback on connection
+	GOnConnection = Desc.OnConnectionFunc;
+
 	// Allow the worker thread to start updating 
 	GUpdateInProgress = 0;
 }
diff --git a/Engine/Source/Runtime/TraceLog/Public/Trace/Trace.h b/Engine/Source/Runtime/TraceLog/Public/Trace/Trace.h
--- a/Engine/Source/Runtime/TraceLog/Public/Trace/Trace.h
+++ b/Engine/Source/Runtime/TraceLog/Public/Trace/Trace.h
@@ -77,11 +77,14 @@
 	return TEventRef<IdType>(InId, InTypeId);
 }
 
+using OnConnectFunc = void(void);
+
 struct FInitializeDesc
 {
 	uint32			TailSizeBytes		= 4 << 20;
 	bool			bUseWorkerThread	= true;
 	bool			bUseImportantCache	= true;
+	OnConnectFunc*	OnConnectionFunc	= nullptr;
 };
 
 typedef void*		AllocFunc(SIZE_T, uint32);
