Make it possible to export Trace events

* Allow events to be exported by prepending UE_TRACE_EVENT_BEGIN_EXTERN macro with *_API macro.
* Changed reference fields to take uid directly in constructor rather than object node pointer.
* Removed unused duplicated event in animation module.

diff --git a/Engine/Source/Runtime/Engine/Private/Animation/AnimTrace.cpp b/Engine/Source/Runtime/Engine/Private/Animation/AnimTrace.cpp
--- a/Engine/Source/Runtime/Engine/Private/Animation/AnimTrace.cpp
+++ b/Engine/Source/Runtime/Engine/Private/Animation/AnimTrace.cpp
@@ -187,16 +187,6 @@ UE_TRACE_EVENT_BEGIN(Animation, AnimSequencePlayer)
 	UE_TRACE_EVENT_FIELD(uint16, FrameCounter)
 UE_TRACE_EVENT_END()
 
-UE_TRACE_EVENT_BEGIN(Animation, BlendSpacePlayer)
-	UE_TRACE_EVENT_FIELD(uint64, Cycle)
-	UE_TRACE_EVENT_FIELD(uint64, AnimInstanceId)
-	UE_TRACE_EVENT_FIELD(uint64, BlendSpaceId)
-	UE_TRACE_EVENT_FIELD(int32, NodeId)
-	UE_TRACE_EVENT_FIELD(float, PositionX)
-	UE_TRACE_EVENT_FIELD(float, PositionY)
-	UE_TRACE_EVENT_FIELD(float, PositionZ)
-UE_TRACE_EVENT_END()
-
 UE_TRACE_EVENT_BEGIN(Animation, StateMachineState)
 	UE_TRACE_EVENT_FIELD(uint64, Cycle)
 	UE_TRACE_EVENT_FIELD(uint64, AnimInstanceId)
diff --git a/Engine/Source/Runtime/TraceLog/Private/Trace/EventNode.cpp b/Engine/Source/Runtime/TraceLog/Private/Trace/EventNode.cpp
--- a/Engine/Source/Runtime/TraceLog/Private/Trace/EventNode.cpp
+++ b/Engine/Source/Runtime/TraceLog/Private/Trace/EventNode.cpp
@@ -170,14 +170,14 @@ void FEventNode::Describe() const
 			Out.DefinitionId.Offset = Field.ValueOffset;
 			Out.DefinitionId.TypeInfo = Field.TypeInfo;
 		}
-		else if (Field.Reference)
+		else if (Field.Reference != 0)
 		{
 			Out.FieldType = EFieldFamily::Reference;
 			// todo: What if the referenced type has not been initialized yet?
 			Out.Reference.Offset = Field.ValueOffset;
 			Out.Reference.TypeInfo = Field.TypeInfo;
 			Out.Reference.NameSize = Field.NameSize;
-			Out.Reference.RefUid = uint16(Field.Reference->GetUid()) >> EKnownEventUids::_UidShift;
+			Out.Reference.RefUid = uint16(Field.Reference) >> EKnownEventUids::_UidShift;
 		}
 		else
 		{
diff --git a/Engine/Source/Runtime/TraceLog/Public/Trace/Detail/Field.h b/Engine/Source/Runtime/TraceLog/Public/Trace/Detail/Field.h
--- a/Engine/Source/Runtime/TraceLog/Public/Trace/Detail/Field.h
+++ b/Engine/Source/Runtime/TraceLog/Public/Trace/Detail/Field.h
@@ -96,13 +96,13 @@ struct FLiteralName
 ////////////////////////////////////////////////////////////////////////////////
 struct FFieldDesc
 {
-	FFieldDesc(const FLiteralName& Name, uint8 Type, uint16 Offset, uint16 Size, Private::FEventNode* ReferenceEvent = nullptr)
+	FFieldDesc(const FLiteralName& Name, uint8 Type, uint16 Offset, uint16 Size, uint32 ReferencedUid = 0)
 	: Name(Name.Ptr)
 	, ValueOffset(Offset)
 	, ValueSize(Size)
 	, NameSize(Name.Length)
 	, TypeInfo(Type)
-	, Reference(ReferenceEvent)
+	, Reference(ReferencedUid)
 	{
 	}
 
@@ -111,7 +111,7 @@ struct FFieldDesc
 	uint16			ValueSize;
 	uint8			NameSize;
 	uint8			TypeInfo;
-	Private::FEventNode* Reference;
+	uint32			Reference;
 };
 
 
@@ -181,9 +181,10 @@ struct TField<InIndex, InOffset, TEventRef<DefinitionType>>
 {
 	TRACE_PRIVATE_FIELD(InIndex, InOffset, DefinitionType);
 public:
-	TField(const FLiteralName& Name, Private::FEventNode* ReferenceEvent)
-		: FieldDesc(Name, Tid, Offset, Size, ReferenceEvent)
+	TField(const FLiteralName& Name, uint32 ReferencedUid)
+		: FieldDesc(Name, Tid, Offset, Size, ReferencedUid)
 	{
+		check(ReferencedUid != 0);
 	}
 };
 	
diff --git a/Engine/Source/Runtime/TraceLog/Public/Trace/Detail/Trace.h b/Engine/Source/Runtime/TraceLog/Public/Trace/Detail/Trace.h
--- a/Engine/Source/Runtime/TraceLog/Public/Trace/Detail/Trace.h
+++ b/Engine/Source/Runtime/TraceLog/Public/Trace/Detail/Trace.h
@@ -49,16 +49,18 @@ class FChannel;
 	bool(ChannelsExpr)
 
 #define TRACE_PRIVATE_EVENT_DEFINE(LoggerName, EventName) \
-	UE::Trace::Private::FEventNode LoggerName##EventName##Event;
+	static UE::Trace::Private::FEventNode LoggerName##EventName##Event##Impl;\
+	UE::Trace::Private::FEventNode& LoggerName##EventName##Event = LoggerName##EventName##Event##Impl;
 
 #define TRACE_PRIVATE_EVENT_BEGIN(LoggerName, EventName, ...) \
-	TRACE_PRIVATE_EVENT_BEGIN_IMPL(static, LoggerName, EventName, ##__VA_ARGS__)
+	TRACE_PRIVATE_EVENT_DEFINE(LoggerName, EventName); \
+	TRACE_PRIVATE_EVENT_BEGIN_IMPL(LoggerName, EventName, ##__VA_ARGS__)
 
 #define TRACE_PRIVATE_EVENT_BEGIN_EXTERN(LoggerName, EventName, ...) \
-	TRACE_PRIVATE_EVENT_BEGIN_IMPL(extern, LoggerName, EventName, ##__VA_ARGS__)
+	extern UE::Trace::Private::FEventNode& LoggerName##EventName##Event; \
+	TRACE_PRIVATE_EVENT_BEGIN_IMPL(LoggerName, EventName, ##__VA_ARGS__)
 
-#define TRACE_PRIVATE_EVENT_BEGIN_IMPL(LinkageType, LoggerName, EventName, ...) \
-	LinkageType TRACE_PRIVATE_EVENT_DEFINE(LoggerName, EventName) \
+#define TRACE_PRIVATE_EVENT_BEGIN_IMPL(LoggerName, EventName, ...) \
 	struct F##LoggerName##EventName##Fields \
 	{ \
 		enum \
@@ -113,7 +115,7 @@ class FChannel;
 
 #define TRACE_PRIVATE_EVENT_REFFIELD(RefLogger, RefEventType, FieldName) \
 		F##RefLogger##RefEventType##Fields::DefinitionType> FieldName##_Meta; \
-		FieldName##_Meta const FieldName##_Field = FieldName##_Meta(UE::Trace::FLiteralName(#FieldName), &RefLogger##RefEventType##Event); \
+		FieldName##_Meta const FieldName##_Field = FieldName##_Meta(UE::Trace::FLiteralName(#FieldName), RefLogger##RefEventType##Event.GetUid()); \
 		template <typename DefinitionType> auto FieldName(UE::Trace::TEventRef<DefinitionType> Reference) const { \
 			checkfSlow(Reference.RefTypeId == F##RefLogger##RefEventType##Fields::GetUid(), TEXT("Incorrect reference type passed to event. Field expected %s with uid %u but got a reference with uid %u"), TEXT(#RefEventType), F##RefLogger##RefEventType##Fields::GetUid(), Reference.RefTypeId);\
 			LogScopeType::FFieldSet<FieldName##_Meta, F##RefLogger##RefEventType##Fields::DefinitionType>::Impl((LogScopeType*)this, Reference); \
@@ -188,12 +190,14 @@ class FChannel;
 #define TRACE_PRIVATE_CHANNELEXPR_IS_ENABLED(ChannelsExpr) \
 	false
 
-#define TRACE_PRIVATE_EVENT_DEFINE(LoggerName, EventName)
+#define TRACE_PRIVATE_EVENT_DEFINE(LoggerName, EventName) \
+	int8* LoggerName##EventName##DummyPtr = nullptr;\
 
 #define TRACE_PRIVATE_EVENT_BEGIN(LoggerName, EventName, ...) \
 	TRACE_PRIVATE_EVENT_BEGIN_IMPL(LoggerName, EventName)
 
 #define TRACE_PRIVATE_EVENT_BEGIN_EXTERN(LoggerName, EventName, ...) \
+	extern int8* LoggerName##EventName##DummyPtr;\
 	TRACE_PRIVATE_EVENT_BEGIN_IMPL(LoggerName, EventName)
 
 #define TRACE_PRIVATE_EVENT_BEGIN_IMPL(LoggerName, EventName) \
