Exclude EventType::GetUid from TSAN checks

diff --git a/Engine/Source/Runtime/TraceLog/Public/Trace/Detail/Trace.h b/Engine/Source/Runtime/TraceLog/Public/Trace/Detail/Trace.h
--- a/Engine/Source/Runtime/TraceLog/Public/Trace/Detail/Trace.h
+++ b/Engine/Source/Runtime/TraceLog/Public/Trace/Detail/Trace.h
@@ -81,7 +81,7 @@ class FChannel;
 		bIsDefinition64 = ((0, ##__VA_ARGS__) & Definition64bit) != 0,}; \
 		typedef std::conditional_t<bIsDefinition8, UE::Trace::FEventRef8, std::conditional_t<bIsDefinition16, UE::Trace::FEventRef16 , std::conditional_t<bIsDefinition64, UE::Trace::FEventRef64, UE::Trace::FEventRef32>>> DefinitionType;\
 		static constexpr uint32 GetSize() { return EventProps_Meta::Size; } \
-		static uint32 GetUid() { static uint32 Uid = 0; return (Uid = Uid ? Uid : Initialize()); } \
+		static uint32 TSAN_SAFE GetUid() { static uint32 Uid = 0; return (Uid = Uid ? Uid : Initialize()); } \
 		static uint32 FORCENOINLINE Initialize() \
 		{ \
 			static const uint32 Uid_ThreadSafeInit = [] () \
