Fixed TSAN-reported error.

#rnx

diff --git a/Engine/Source/Runtime/TraceLog/Private/Trace/TlsBuffer.cpp b/Engine/Source/Runtime/TraceLog/Private/Trace/TlsBuffer.cpp
--- a/Engine/Source/Runtime/TraceLog/Private/Trace/TlsBuffer.cpp
+++ b/Engine/Source/Runtime/TraceLog/Private/Trace/TlsBuffer.cpp
@@ -122,7 +122,7 @@ TRACELOG_API FWriteBuffer* Writer_NextBuffer(int32 Size)
 ////////////////////////////////////////////////////////////////////////////////
 static bool Writer_DrainBuffer(uint32 ThreadId, FWriteBuffer* Buffer)
 {
-	uint8* Committed = AtomicLoadRelaxed((uint8**)&Buffer->Committed);
+	uint8* Committed = AtomicLoadAcquire((uint8**)&Buffer->Committed);
 
 	// Send as much as we can.
 	if (uint32 SizeToReap = uint32(Committed - Buffer->Reaped))
