Fixed trace data getting sent out-of-order.

#rnx

diff --git a/Engine/Source/Runtime/TraceLog/Private/Trace/Tail.cpp b/Engine/Source/Runtime/TraceLog/Private/Trace/Tail.cpp
--- a/Engine/Source/Runtime/TraceLog/Private/Trace/Tail.cpp
+++ b/Engine/Source/Runtime/TraceLog/Private/Trace/Tail.cpp
@@ -40,7 +40,6 @@ public:
 	void			Initialize(uint32 InSize);
 	void			Shutdown();
 	void			Reset();
-	void			Pause(bool bInPause) { bPaused = bInPause; };
 	uint32			GetSize() const;
 	bool			IsActive() const;
 	FRange			GetBackPackets() const;
@@ -58,7 +57,6 @@ private:
 	uint32			Cursor;
 	uint32			Left;
 	uint32			Right;
-	bool			bPaused;
 };
 
 // TraceLog must be ready after value-init so that it can be used before
@@ -70,7 +68,6 @@ void FPacketRing::Initialize(uint32 InSize)
 {
 	Data = (uint8*)Writer_MemoryAllocate(InSize, 16);
 	Size = InSize;
-	bPaused = false;
 	Reset();
 }
 
@@ -97,7 +94,7 @@ uint32 FPacketRing::GetSize() const
 ////////////////////////////////////////////////////////////////////////////////
 bool FPacketRing::IsActive() const
 {
-	return Data != nullptr && !bPaused;
+	return Data != nullptr;
 }
 
 ////////////////////////////////////////////////////////////////////////////////
@@ -201,13 +198,6 @@ FTidPacketBase* FPacketRing::AppendImpl(uint32 InSize)
 ////////////////////////////////////////////////////////////////////////////////
 static FPacketRing GPacketRing; // = {};
 
-
-////////////////////////////////////////////////////////////////////////////////
-void Writer_TailPause(bool bPause)
-{
-	GPacketRing.Pause(bPause);
-}
-
 ////////////////////////////////////////////////////////////////////////////////
 void Writer_TailAppend(uint32 ThreadId, uint8* __restrict Data, uint32 Size, bool bPartial)
 {
diff --git a/Engine/Source/Runtime/TraceLog/Private/Trace/Writer.cpp b/Engine/Source/Runtime/TraceLog/Private/Trace/Writer.cpp
--- a/Engine/Source/Runtime/TraceLog/Private/Trace/Writer.cpp
+++ b/Engine/Source/Runtime/TraceLog/Private/Trace/Writer.cpp
@@ -41,7 +41,6 @@ void			Writer_SendData(uint32, uint8* __restrict, uint32);
 void			Writer_InitializeTail(int32);
 void			Writer_ShutdownTail();
 void			Writer_TailOnConnect();
-void			Writer_TailPause(bool);	
 void			Writer_InitializeSharedBuffers();
 void			Writer_ShutdownSharedBuffers();
 void			Writer_UpdateSharedBuffers();
@@ -58,7 +57,6 @@ uint32			Writer_GetControlPort();
 void			Writer_UpdateControl();
 void			Writer_InitializeControl();
 void			Writer_ShutdownControl();
-bool			Writer_IsTracing();
 bool			Writer_IsTailing();
 static bool		Writer_SessionPrologue();
 void			Writer_FreeBlockListToPool(FWriteBuffer*, FWriteBuffer*);
@@ -551,17 +549,18 @@ void Writer_CallbackOnConnect()
 	{
 		return;
 	}
-	
-	Writer_TailPause(true);
+
+	// Prior to letting callbacks trace events we need to flush any pending
+	// trace data to that tail. We do not want that data to be sent over the
+	// wire as that would cause data to be sent out-of-order.
+	UPTRINT DataHandle = GDataHandle;
+	GDataHandle = 0;
+	Writer_DrainLocalBuffers();
+	GDataHandle = DataHandle;
 
 	// Issue callback. We assume any events emitted here are not marked as
 	// important and emitted on this thread.
 	GOnConnection();
-
-	// Drain only the events for this thread passing tail as you go.
-	Writer_DrainLocalBuffers();
-
-	Writer_TailPause(false);
 }
 
 
