Cleaned up spurious trailing backslashes

#rnx

diff --git a/Engine/Source/Runtime/TraceLog/Public/Trace/Detail/Trace.h b/Engine/Source/Runtime/TraceLog/Public/Trace/Detail/Trace.h
--- a/Engine/Source/Runtime/TraceLog/Public/Trace/Detail/Trace.h
+++ b/Engine/Source/Runtime/TraceLog/Public/Trace/Detail/Trace.h
@@ -191,13 +191,13 @@ class FChannel;
 	false
 
 #define TRACE_PRIVATE_EVENT_DEFINE(LoggerName, EventName) \
-	int8* LoggerName##EventName##DummyPtr = nullptr;\
+	int8* LoggerName##EventName##DummyPtr = nullptr;
 
 #define TRACE_PRIVATE_EVENT_BEGIN(LoggerName, EventName, ...) \
 	TRACE_PRIVATE_EVENT_BEGIN_IMPL(LoggerName, EventName)
 
 #define TRACE_PRIVATE_EVENT_BEGIN_EXTERN(LoggerName, EventName, ...) \
-	extern int8* LoggerName##EventName##DummyPtr;\
+	extern int8* LoggerName##EventName##DummyPtr; \
 	TRACE_PRIVATE_EVENT_BEGIN_IMPL(LoggerName, EventName)
 
 #define TRACE_PRIVATE_EVENT_BEGIN_IMPL(LoggerName, EventName) \
@@ -239,6 +239,6 @@ class FChannel;
 
 #define TRACE_PRIVATE_LOG_DEFINITION(LoggerName, EventName, Id, ChannelsExpr, ...) \
 	UE::Trace::MakeEventRef(Id, 0); \
-	TRACE_PRIVATE_LOG(LoggerName, EventName, ChannelsExpr, ##__VA_ARGS__) \
+	TRACE_PRIVATE_LOG(LoggerName, EventName, ChannelsExpr, ##__VA_ARGS__)
 
 #endif // UE_TRACE_ENABLED
