Deleted UTF8 BOM

#rnx

diff --git a/Engine/Source/Runtime/TraceLog/Public/Trace/Detail/Protocols/Protocol6.h b/Engine/Source/Runtime/TraceLog/Public/Trace/Detail/Protocols/Protocol6.h
--- a/Engine/Source/Runtime/TraceLog/Public/Trace/Detail/Protocols/Protocol6.h
+++ b/Engine/Source/Runtime/TraceLog/Public/Trace/Detail/Protocols/Protocol6.h
@@ -1,4 +1,4 @@
-﻿// Copyright Epic Games, Inc. All Rights Reserved.
+// Copyright Epic Games, Inc. All Rights Reserved.
 #pragma once
 
 #include "Protocol4.h"
