Support creating standalone analysis too.

#rnx

diff --git a/Engine/Source/Runtime/TraceLog/create_standalone.py b/Engine/Source/Runtime/TraceLog/create_standalone.py
--- a/Engine/Source/Runtime/TraceLog/create_standalone.py
+++ b/Engine/Source/Runtime/TraceLog/create_standalone.py
@@ -40,7 +40,6 @@ def _exclude_line(line):
 
 #-------------------------------------------------------------------------------
 def _collect_source(src_dir, predicate=None):
-	_spam_header("Gathering source files")
 	files  = [x for x in src_dir.glob("Public/**/*")]
 	files += [x for x in src_dir.glob("Private/**/*")]
 	files  = [x for x in files if x.suffix in (".h", ".cpp", ".inl")]
@@ -50,12 +49,14 @@ def _collect_source(src_dir, predicate=None):
 		files = [x for x in files if predicate(x)]
 		_spam("Filtered down to %d files" % len(files))
 
-	# Collect lines of each source file, filtering local includes
-	source_files = []
-	for file in files:
-		source_file = _SourceFile(file)
-		source_files.append(source_file)
+	return [_SourceFile(x) for x in files]
 
+#-------------------------------------------------------------------------------
+def _finalize_source(source_files, include_dirs):
+	# Collect lines of each source file, filtering local includes
+	_spam("Loading code")
+	for source_file in source_files:
+		file = source_file.path
 		for line in file.open("rt"):
 			include = _parse_include(line)
 			if not include:
@@ -65,7 +66,7 @@ def _collect_source(src_dir, predicate=None):
 
 			is_local = False
 
-			for include_dir in (file.parent, src_dir / "Public"):
+			for include_dir in (file.parent, *include_dirs):
 				candidate = include_dir / include
 				if candidate.is_file():
 					is_local = True
@@ -74,8 +75,6 @@ def _collect_source(src_dir, predicate=None):
 			if is_local:
 				source_file.deps.append(candidate)
 
-	_spam_header("Topological sort")
-
 	# Hook up dependencies
 	_spam("Resolving include dependencies")
 	for source_file in source_files:
@@ -108,21 +107,51 @@ def _collect_source(src_dir, predicate=None):
 	ext_key = lambda x: 1 if x.path.suffix == ".cpp" else 0
 	source_files = sorted(topo_sorted, key=ext_key)
 
-	# Exclude LZ4
-	source_files = [x for x in source_files if "lz4" not in x.path.name]
-
 	return source_files
 
 #-------------------------------------------------------------------------------
-def _main(src_dir, dest_dir, thin):
+def _main_trace(src_dir, dest_dir, thin, analysis):
 	dest_dir = dest_dir.resolve()
 	src_dir = Path(__file__).parent
 
-	_spam_header("Trace")
 	_spam("Source dir:", src_dir)
 	_spam("Dest dir:", dest_dir)
 
-	source_files = _collect_source(src_dir)
+	# Collect and sort source files.
+	def predicate(path):
+		if path.stem.startswith("lz4"): return False
+		return True
+	source_files = _collect_source(src_dir, predicate)
+
+	analysis_dir = None
+	if analysis:
+		analysis_dir = Path(__file__).parent
+		analysis_dir = analysis_dir.parent.parent / "Developer/TraceAnalysis"
+		_spam("Analysis dir:", analysis_dir)
+
+		def predicate(path):
+			if	 path.parent.name == "Store":		 return False
+			elif path.parent.name == "Asio":		 return False
+			elif path.name.startswith("Store"):		 return False
+			elif path.name.startswith("Control"):	 return False
+			excludes = (
+				"Analysis.h",
+				"Context.cpp",
+				"DataStream.h",
+				"DataStream.cpp",
+				"EventToCbor.cpp",
+				"Processor.cpp",
+				"Processor.h",
+				"TraceAnalysisModule.cpp",
+			)
+			return path.name not in excludes
+		source_files += _collect_source(analysis_dir, predicate)
+
+	_spam_header("Topological sort")
+	include_dirs = (src_dir / "Public",)
+	if analysis_dir:
+		include_dirs = (*include_dirs, analysis_dir / "Public")
+	source_files = _finalize_source(source_files, include_dirs)
 
 	# Add prologue and epilogue files
 	prologue = _SourceFile(src_dir / "standalone_prologue.h")
@@ -183,8 +212,10 @@ def main():
 	parser = argparse.ArgumentParser(description=desc)
 	parser.add_argument("outdir", help="Directory to write output file(s) to")
 	parser.add_argument("--thin", action="store_true", help="#include standalone_ instead of blitting")
+	parser.add_argument("--analysis", action="store_true", help="Amalgamate TraceAnalysis too")
 	args = parser.parse_args()
-	return _main(Path(__file__).parent, Path(args.outdir), args.thin)
+
+	_main_trace(Path(__file__).parent, Path(args.outdir), args.thin, args.analysis)
 
 if __name__ == "__main__":
 	raise SystemExit(main())
diff --git a/Engine/Source/Runtime/TraceLog/standalone_epilogue.h b/Engine/Source/Runtime/TraceLog/standalone_epilogue.h
--- a/Engine/Source/Runtime/TraceLog/standalone_epilogue.h
+++ b/Engine/Source/Runtime/TraceLog/standalone_epilogue.h
@@ -1,11 +1,28 @@
 // Copyright Epic Games, Inc. All Rights Reserved.
 
+// {{{1 amalgamation-tail //////////////////////////////////////////////////////
+
 #if PLATFORM_WINDOWS
 #	pragma warning(pop)
 #endif
 
 #if TRACE_UE_COMPAT_LAYER
 
+namespace UE {
+namespace Trace {
+
+#if defined(TRACE_HAS_ANALYSIS)
+
+inline void SerializeToCborImpl(TArray<uint8>&, const IAnalyzer::FEventData&, uint32)
+{
+	*(int*)0 = 0; // unsupported
+}
+
+#endif // TRACE_HAS_ANALYSIS
+
+} // namespace Trace
+} // namespace UE
+
 #if PLATFORM_WINDOWS
 #	if defined(UNICODE) || defined(_UNICODE)
 #		undef TEXT
@@ -16,6 +33,10 @@
 
 #endif // TRACE_UE_COMPAT_LAYER
 
+
+
+// {{{1 library-setup //////////////////////////////////////////////////////////
+
 #include <string_view>
 
 #define TRACE_EVENT_DEFINE			UE_TRACE_EVENT_DEFINE
@@ -36,15 +57,108 @@ namespace trace = UE::Trace;
 #define TRACE_PRIVATE_CONCAT(x, y)		TRACE_PRIVATE_CONCAT_(x, y)
 #define TRACE_PRIVATE_UNIQUE_VAR(name)	TRACE_PRIVATE_CONCAT($trace_##name, __LINE__)
 
+
+
+// {{{1 session-header /////////////////////////////////////////////////////////
+
+namespace UE {
+namespace Trace {
+
+enum class Build
+{
+	Unknown,
+	Debug,
+	DebugGame,
+	Development,
+	Shipping,
+	Test
+};
+
+void DescribeSession(
+	const std::string_view& AppName,
+	Build Variant=Build::Unknown,
+	const std::string_view& CommandLine="",
+	const std::string_view& BuildVersion="unknown_ver");
+
+} // namespace Trace
+} // namespace UE
+
+// {{{1 session-source /////////////////////////////////////////////////////////
+
+#if TRACE_IMPLEMENT
+
+namespace UE {
+namespace Trace {
+namespace Private {
+
+TRACE_EVENT_BEGIN(Diagnostics, Session2, NoSync|Important)
+	TRACE_EVENT_FIELD(uint8, ConfigurationType)
+	TRACE_EVENT_FIELD(trace::AnsiString, AppName)
+	TRACE_EVENT_FIELD(trace::AnsiString, BuildVersion)
+	TRACE_EVENT_FIELD(trace::AnsiString, Platform)
+	TRACE_EVENT_FIELD(trace::AnsiString, CommandLine)
+TRACE_EVENT_END()
+
+} // namespace Private
+
+////////////////////////////////////////////////////////////////////////////////
+void DescribeSession(
+	const std::string_view& AppName,
+	Build Variant,
+	const std::string_view& CommandLine,
+	const std::string_view& BuildVersion)
+{
+	using namespace Private;
+	using namespace std::literals;
+
+	std::string_view Platform;
+#if PLATFORM_WINDOWS
+	Platform = "Windows"sv;
+#elif PLATFORM_UNIX
+	Platform = "Linux"sv;
+#elif PLATFORM_MAC
+	Platform = "Mac"sv;
+#else
+	Platform = "Unknown"sv;
+#endif
+
+	int DataSize = 0;
+	DataSize += AppName.size();
+	DataSize += BuildVersion.size();
+	DataSize += Platform.size();
+	DataSize += CommandLine.size();
+
+	TRACE_LOG(Diagnostics, Session2, true, DataSize)
+		<< Session2.AppName(AppName.data(), int(AppName.size()))
+		<< Session2.BuildVersion(BuildVersion.data(), int(BuildVersion.size()))
+		<< Session2.Platform(Platform.data(), int(Platform.size()))
+		<< Session2.CommandLine(CommandLine.data(), int(CommandLine.size()))
+		<< Session2.ConfigurationType(uint8(Variant));
+}
+
+} // namespace Trace
+} // namespace UE
+
+#endif // TRACE_IMPLEMENT
+
+
+
+// {{{1 cpu-header /////////////////////////////////////////////////////////////
+
 TRACE_CHANNEL_EXTERN(CpuChannel)
 
 namespace UE {
 namespace Trace {
 
+enum CpuScopeFlags : int
+{
+	CpuFlush = 1 << 0,
+};
+
 struct TraceCpuScope
 {
 			~TraceCpuScope();
-	void	Enter(int ScopeId);
+	void	Enter(int ScopeId, int Flags=0);
 	int		_ScopeId = 0;
 };
 
@@ -53,54 +167,44 @@ int ScopeNew(const std::string_view& Name);
 } // namespace Trace
 } // namespace UE
 
-#define TRACE_CPU_SCOPE(name) \
-	using namespace std::literals; \
+#define TRACE_CPU_SCOPE(name, ...) \
 	trace::TraceCpuScope TRACE_PRIVATE_UNIQUE_VAR(cpu_scope); \
 	if (CpuChannel) { \
+		using namespace std::literals; \
 		static int TRACE_PRIVATE_UNIQUE_VAR(scope_id); \
 		if (0 == TRACE_PRIVATE_UNIQUE_VAR(scope_id)) \
 			TRACE_PRIVATE_UNIQUE_VAR(scope_id) = trace::ScopeNew(name##sv); \
-		TRACE_PRIVATE_UNIQUE_VAR(cpu_scope).Enter(TRACE_PRIVATE_UNIQUE_VAR(scope_id)); \
+		TRACE_PRIVATE_UNIQUE_VAR(cpu_scope).Enter(TRACE_PRIVATE_UNIQUE_VAR(scope_id), ##__VA_ARGS__); \
 	} \
 	do {} while (0)
 
+// {{{1 cpu-source /////////////////////////////////////////////////////////////
+
 #if TRACE_IMPLEMENT
 
-////////////////////////////////////////////////////////////////////////////////
 TRACE_CHANNEL_DEFINE(CpuChannel)
 
+namespace UE {
+namespace Trace {
+namespace Private {
+
 TRACE_EVENT_BEGIN(CpuProfiler, EventSpec, NoSync|Important)
 	TRACE_EVENT_FIELD(uint32, Id)
-	TRACE_EVENT_FIELD(UE::Trace::AnsiString, Name)
+	TRACE_EVENT_FIELD(trace::AnsiString, Name)
 TRACE_EVENT_END()
 
 TRACE_EVENT_BEGIN(CpuProfiler, EventBatch, NoSync)
 	TRACE_EVENT_FIELD(uint8[], Data)
 TRACE_EVENT_END()
 
-namespace UE {
-namespace Trace {
-namespace Private {
-
+////////////////////////////////////////////////////////////////////////////////
 static int32_t encode32_7bit(int32_t value, void* __restrict out)
 {
 	// Calculate the number of bytes
-#if 0
-	int32_t msb_test = (value << sizeof(value)) | 0x10;
-#if _MSC_VER
-	unsigned long bit_index;
-	_BitScanReverse(&bit_index, msb_test);
-#else
-	int32_t leading_zeros = __builtin_clz(msb_test);
-	int32_t bit_index = ((sizeof(value) * 8) - 1) - leading_zeros;
-#endif
-	int32_t length = (bit_index + 3) / 7;
-#else
 	int32_t length = 1;
 	length += (value >= (1 <<  7));
 	length += (value >= (1 << 14));
 	length += (value >= (1 << 21));
-#endif
 
 	// Add a gap every eigth bit for the continuations
 	int32_t ret = value;
@@ -117,20 +221,10 @@ static int32_t encode32_7bit(int32_t value, void* __restrict out)
 	return length;
 }
 
+////////////////////////////////////////////////////////////////////////////////
 static int32_t encode64_7bit(int64_t value, void* __restrict out)
 {
 	// Calculate the output length
-#if 0
-	int64_t msb_test = (value << sizeof(value)) | 0x100ull;
-#if _MSC_VER
-	unsigned long bit_index;
-	_BitScanReverse64(&bit_index, msb_test);
-#else
-	int32_t leading_zeros = __builtin_clzll(msb_test);
-	int32_t bit_index = ((sizeof(value) * 8) - 1) - leading_zeros;
-#endif
-	int32_t length = (bit_index - 1) / 7;
-#else
 	uint32_t length = 1;
 	length += (value >= (1ll <<  7));
 	length += (value >= (1ll << 14));
@@ -139,7 +233,6 @@ static int32_t encode64_7bit(int64_t value, void* __restrict out)
 	length += (value >= (1ll << 35));
 	length += (value >= (1ll << 42));
 	length += (value >= (1ll << 49));
-#endif
 
 	// Add a gap every eigth bit for the continuations
 	int64_t ret = value;
@@ -161,13 +254,13 @@ static int32_t encode64_7bit(int64_t value, void* __restrict out)
 class ThreadBuffer
 {
 public:
-	static void	Enter(uint64_t Timestamp, uint32_t ScopeId) { TlsInstance.EnterImpl(Timestamp, ScopeId); }
-	static void	Leave(uint64_t Timestamp)					{ TlsInstance.LeaveImpl(Timestamp); }
+	static void	Enter(uint64_t Timestamp, uint32_t ScopeId, int Flags);
+	static void	Leave(uint64_t Timestamp);
 
 private:
 				~ThreadBuffer();
 	void		Flush(bool Force);
-	void		EnterImpl(uint64_t Timestamp, uint32_t ScopeId);
+	void		EnterImpl(uint64_t Timestamp, uint32_t ScopeId, int Flags);
 	void		LeaveImpl(uint64_t Timestamp);
 	enum
 	{
@@ -185,6 +278,18 @@ private:
 
 thread_local ThreadBuffer ThreadBuffer::TlsInstance;
 
+////////////////////////////////////////////////////////////////////////////////
+inline void	ThreadBuffer::Enter(uint64_t Timestamp, uint32_t ScopeId, int Flags)
+{
+	TlsInstance.EnterImpl(Timestamp, ScopeId, Flags);
+}
+
+////////////////////////////////////////////////////////////////////////////////
+inline void	ThreadBuffer::Leave(uint64_t Timestamp)
+{
+	TlsInstance.LeaveImpl(Timestamp);
+}
+
 ////////////////////////////////////////////////////////////////////////////////
 ThreadBuffer::~ThreadBuffer()
 {
@@ -194,6 +299,8 @@ ThreadBuffer::~ThreadBuffer()
 ////////////////////////////////////////////////////////////////////////////////
 void ThreadBuffer::Flush(bool Force)
 {
+	using namespace Private;
+
 	if (!Force && (Cursor <= (Buffer + BufferSize - Overflow)))
 		return;
 
@@ -205,14 +312,15 @@ void ThreadBuffer::Flush(bool Force)
 }
 
 ////////////////////////////////////////////////////////////////////////////////
-void ThreadBuffer::EnterImpl(uint64_t Timestamp, uint32_t ScopeId)
+void ThreadBuffer::EnterImpl(uint64_t Timestamp, uint32_t ScopeId, int Flags)
 {
 	Timestamp -= PrevTimestamp;
 	PrevTimestamp += Timestamp;
 	Cursor += encode64_7bit((Timestamp) << 1 | EnterLsb, Cursor);
 	Cursor += encode32_7bit(ScopeId, Cursor);
 
-	Flush(false);
+	bool ShouldFlush = (Flags & CpuScopeFlags::CpuFlush);
+	Flush(ShouldFlush);
 }
 
 ////////////////////////////////////////////////////////////////////////////////
@@ -232,8 +340,10 @@ void ThreadBuffer::LeaveImpl(uint64_t Timestamp)
 ////////////////////////////////////////////////////////////////////////////////
 int ScopeNew(const std::string_view& Name)
 {
+	using namespace Private;
+
 	static int volatile NextSpecId = 1;
-	int SpecId = Private::AtomicAddRelaxed(&NextSpecId, 1);
+	int SpecId = AtomicAddRelaxed(&NextSpecId, 1);
 
 	uint32 NameSize = uint32(Name.size());
 	TRACE_LOG(CpuProfiler, EventSpec, true, NameSize)
@@ -258,15 +368,124 @@ TraceCpuScope::~TraceCpuScope()
 }
 
 ////////////////////////////////////////////////////////////////////////////////
-void TraceCpuScope::Enter(int ScopeId)
+void TraceCpuScope::Enter(int ScopeId, int Flags)
 {
 	using namespace Private;
 
 	_ScopeId = ScopeId;
 	uint64 Timestamp = TimeGetTimestamp();
-	ThreadBuffer::Enter(Timestamp, ScopeId);
+	ThreadBuffer::Enter(Timestamp, ScopeId, Flags);
+}
+
+} // namespace Trace
+} // namespace UE
+
+#endif // TRACE_IMPLEMENT
+
+
+
+// {{{1 log-header /////////////////////////////////////////////////////////////
+
+TRACE_CHANNEL_EXTERN(LogChannel)
+
+namespace UE {
+namespace Trace {
+namespace Private {
+
+void LogMessageImpl(int Id, const void* ParamBuffer, int ParamSize);
+int	 LogMessageNew(const std::string_view& Format, const std::string_view& File, int Line);
+
+template <typename... ARGS>
+void LogMessage(int Id, ARGS&&... Args)
+{
+	LogMessageImpl(Id, nullptr, 0);
+}
+
+} // namespace Private
+} // namespace Trace
+} // namespace UE
+
+#define TRACE_LOG_MESSAGE(format, ...) \
+	if (LogChannel) { \
+		using namespace std::literals; \
+		static int message_id; \
+		if (message_id == 0) \
+			message_id = trace::Private::LogMessageNew( \
+				format##sv, \
+				TRACE_PRIVATE_CONCAT(__FILE__, sv), \
+				__LINE__); \
+		trace::Private::LogMessage(message_id, ##__VA_ARGS__); \
+	} \
+	do {} while (0)
+
+// {{{1 log-source /////////////////////////////////////////////////////////////
+
+#if TRACE_IMPLEMENT
+
+TRACE_CHANNEL_DEFINE(LogChannel)
+
+namespace UE {
+namespace Trace {
+namespace Private {
+
+#if 0
+TRACE_EVENT_BEGIN(Logging, LogCategory, NoSync|Important)
+	TRACE_EVENT_FIELD(const void*, CategoryPointer)
+	TRACE_EVENT_FIELD(uint8, DefaultVerbosity)
+	TRACE_EVENT_FIELD(UE::Trace::AnsiString, Name)
+TRACE_EVENT_END()
+#endif
+
+TRACE_EVENT_BEGIN(Logging, LogMessageSpec, NoSync|Important)
+	TRACE_EVENT_FIELD(uint32, LogPoint)
+	//TRACE_EVENT_FIELD(uint16, CategoryPointer)
+	TRACE_EVENT_FIELD(uint16, Line)
+	TRACE_EVENT_FIELD(UE::Trace::AnsiString, FileName)
+	TRACE_EVENT_FIELD(UE::Trace::AnsiString, FormatString)
+TRACE_EVENT_END()
+
+TRACE_EVENT_BEGIN(Logging, LogMessage, NoSync)
+	TRACE_EVENT_FIELD(uint32, LogPoint)
+	TRACE_EVENT_FIELD(uint64, Cycle)
+	TRACE_EVENT_FIELD(uint8[], FormatArgs)
+TRACE_EVENT_END()
+
+////////////////////////////////////////////////////////////////////////////////
+void LogMessageImpl(int Id, const void* ParamBuffer, int ParamSize)
+{
+	(void)ParamBuffer;
+	(void)ParamSize;
+
+	uint64 Timestamp = TimeGetTimestamp();
+
+	TRACE_LOG(Logging, LogMessage, true)
+		<< LogMessage.LogPoint(Id)
+		<< LogMessage.Cycle(Timestamp);
 }
 
+////////////////////////////////////////////////////////////////////////////////
+int LogMessageNew(
+	const std::string_view& Format,
+	const std::string_view& File,
+	int Line)
+{
+	static int volatile NextId = 1;
+	int Id = AtomicAddRelaxed(&NextId, 1);
+
+	int DataSize = 0;
+	DataSize += Format.size();
+	DataSize += File.size();
+
+	TRACE_LOG(Logging, LogMessageSpec, true, DataSize)
+		<< LogMessageSpec.LogPoint(Id)
+		<< LogMessageSpec.Line(uint16(Line))
+		<< LogMessageSpec.FileName(File.data(), int(File.size()))
+		<< LogMessageSpec.FormatString(Format.data(), int(Format.size()));
+
+	return Id;
+}
+
+} // namespace Private
 } // namespace Trace
 } // namespace UE
 
diff --git a/Engine/Source/Runtime/TraceLog/standalone_prologue.h b/Engine/Source/Runtime/TraceLog/standalone_prologue.h
--- a/Engine/Source/Runtime/TraceLog/standalone_prologue.h
+++ b/Engine/Source/Runtime/TraceLog/standalone_prologue.h
@@ -52,9 +52,12 @@
 #endif
 
 // external includes
+#include <algorithm>
 #include <cstddef>
 #include <cstdint>
+#include <string_view>
 #include <utility>
+#include <vector>
 
 #if PLATFORM_WINDOWS
 #	if !defined(WIN32_LEAN_AND_MEAN)
@@ -91,7 +94,7 @@ using SIZE_T = size_t;
 #define TEXT(x)	x
 #define TCHAR	ANSICHAR
 using ANSICHAR	= char;
-using WIDECHAR	= wchar_t;
+using WIDECHAR	= char16_t;
 
 // keywords
 #if defined(_MSC_VER)
@@ -126,12 +129,14 @@ using WIDECHAR	= wchar_t;
 #else
 #	define TRACELOG_API
 #endif
+#define TRACEANALYSIS_API		TRACELOG_API
 
 #if !defined(IS_MONOLITHIC)
 #	define IS_MONOLITHIC				0
 #endif
 
 // misc defines
+#define TSAN_SAFE
 #define THIRD_PARTY_INCLUDES_END
 #define THIRD_PARTY_INCLUDES_START
 #define TRACE_ENABLED					1
@@ -139,14 +144,168 @@ using WIDECHAR	= wchar_t;
 #define TRACE_PRIVATE_EXTERNAL_LZ4		1
 #define UE_TRACE_ENABLED				TRACE_ENABLED
 #define UE_BUILD_SHIPPING				1
+#define DEFINE_LOG_CATEGORY_STATIC(...)
+#define UE_LOG(...)
+#define check(x)						do { if (!(x)) __debugbreak(); } while (0)
+#define checkf(x, ...)					check(x)
 
 // api
+template <typename T> inline auto Forward(T t)		{ return std::forward<T>(t); }
+template <typename T> inline auto MoveTemp(T&& t)	{ return std::move<T>(t); }
+template <typename T> inline void Swap(T& x, T& y)	{ return std::swap(x, y); }
+
+
+
+template <typename T, typename ALLOCATOR=void>
+struct TArray
+	: public std::vector<T>
+{
+	using	 std::vector<T>::vector;
+	void	 Add(const T& t)					{ push_back(t); }
+	T&		 Add_GetRef(const T& t)				{ push_back(t); return back(); }
+	size_t	 AddUninitialized(size_t n)			{ resize(size() + n); return size() - n; }
+	void	 Emplace()							{ emplace_back(); }
+	T&		 Emplace_GetRef()					{ emplace_back(); return back(); }
+	void	 Insert(T const& t, size_t i)		{ insert(begin() + i, t); }
+	void	 Push(const T& t)					{ push_back(t); }
+	T		 Pop(bool Eh=true)					{ T t = back(); pop_back(); return t; }
+	size_t	 Num() const						{ return size(); }
+	void	 SetNum(size_t num)					{ resize(num); }
+	void	 SetNumZeroed(size_t num)			{ resize(num, 0); }
+	void	 SetNumUninitialized(size_t num)	{ resize(num); }
+	void	 Reserve(size_t num)				{ reserve(num); }
+	T&		 Last()								{ return back(); }
+	T*		 GetData()							{ return data(); }
+	T const* GetData() const					{ return data(); }
+	bool	 IsEmpty() const					{ return empty(); }
+	void	 Empty()							{ clear(); }
+	void	 Reset()							{ clear(); }
+	void	 RemoveAll(bool (*p)(T const&))		{ erase(remove_if(begin(), end(), p), end()); }
+	void	 RemoveAllSwap(bool (*p)(T const&))	{ RemoveAll(p); }
+
+	struct NotLess {
+		bool operator () (T const& Lhs, T const& Rhs) const {
+			return bool(Lhs < Rhs) == false;
+		}
+	};
+	void	 Heapify()				{ std::make_heap(begin(), end(), NotLess()); }
+	void	 HeapPush(T const& t)	{ Add(t); std::push_heap(begin(), end(), NotLess()); }
+	void	 HeapPopDiscard()		{ std::pop_heap(begin(), end(), NotLess()); pop_back(); }
+	T const& HeapTop() const		{ return front(); }
+};
+
+template <size_t S> struct TInlineAllocator {};
+
+template <typename T>
+struct TArrayView
+{
+				TArrayView() = default;
+				TArrayView(T* data, size_t size) : _data(data), _size(size) {}
+	T*			GetData()		{ return _data; }
+	T const*	GetData() const	{ return _data; }
+	size_t		Num() const		{ return _size; }
+	T*			_data = nullptr;
+	size_t		_size = 0;
+	using		value_type = T;
+};
+
+
+template <typename T>
+struct TStringViewAdapter
+	: public std::basic_string_view<T>
+{
+	using		std::basic_string_view<T>::basic_string_view;
+	size_t		Len() const		{ return size(); }
+	const char*	GetData() const	{ return data(); }
+};
+
+struct FStringProxy
+{
+							FStringProxy() = default;
+	template <typename T>	FStringProxy(size_t len, const T* ptr);
+	template <typename T>	FStringProxy(const T& r) : FStringProxy(r.size(), r.data()) {}
+	std::string				_str;
+	const void*				_ptr;
+	uint32_t				_len;
+	uint32_t				_char_size;
+	std::string_view		ToView();
+};
+
+template <typename T>
+inline FStringProxy::FStringProxy(size_t len, const T* ptr)
+: _ptr(ptr)
+, _len(uint32_t(len))
+, _char_size(sizeof(T))
+{
+}
+
+inline std::string_view FStringProxy::ToView()
+{
+	if (_char_size != 1)
+	{
+		_char_size = 1;
+		_str.reserve(_len);
+		char* __restrict cursor = _str.data();
+		for (uint32_t i = 0, n = _len; i < n; ++i)
+			cursor[i] = char(((const char16_t*)_ptr)[i]);
+		_ptr = cursor;
+	}
+
+	return std::string_view((const char*)_ptr, _len);
+}
+
+using FAnsiStringView	= TStringViewAdapter<char>;
+using FWideStringView	= TStringViewAdapter<char16_t>;
+using FString			= struct FStringProxy;
+using ANSITEXTVIEW		= FAnsiStringView;
+
+
+
+namespace UECompat
+{
+
 template <typename T>
-inline auto Forward(T t)
+inline T		Max(T const& a, T const& b)				{ return std::max(a, b); }
+inline size_t	Strlen(const char* str)					{ return ::strlen(str); }
+inline void*	Malloc(size_t size)						{ return ::malloc(size); }
+inline void*	Realloc(void* ptr, size_t size)			{ return ::realloc(ptr, size); }
+inline void		Free(void* ptr)							{ return ::free(ptr); }
+inline void		Memcpy(void* d, const void* s, size_t n){ ::memcpy(d, s, n); }
+
+template <typename T, typename COMPARE=std::less<T::value_type>>
+inline void Sort(T& container, COMPARE&& compare=COMPARE())
 {
-	return std::forward<T>(t);
+	std::sort(container.GetData(), container.GetData() + container.Num(),
+		std::forward<COMPARE>(compare));
 }
 
+template <typename T, typename P>
+inline void SortBy(T& container, P const& p)
+{
+	Sort(container, [&p] (T::value_type& lhs, T::value_type& rhs) {
+		return p(lhs) < p(rhs);
+	});
+}
+
+template <typename T, typename V, typename P>
+inline int LowerBoundBy(T const& container, V const& v, P const& p)
+{
+	auto adapter = [&p] (T::value_type const& lhs, V const& v) {
+		return p(lhs) < v;
+	};
+	auto* first = container.GetData();
+	auto* last = first + container.Num();
+	auto iter = std::lower_bound(first, last, v, adapter);
+	return int(ptrdiff_t(iter - first));
+}
+
+} // namespace UeCompat
+
+namespace Algo			= UECompat;
+namespace FMath			= UECompat;
+namespace FCStringAnsi	= UECompat;
+namespace FMemory		= UECompat;
+
 #endif // TRACE_UE_COMPAT_LAYER
 
 #include <cstring>
