Extend --thin to include all source files

diff --git a/Engine/Source/Runtime/TraceLog/create_standalone.py b/Engine/Source/Runtime/TraceLog/create_standalone.py
--- a/Engine/Source/Runtime/TraceLog/create_standalone.py
+++ b/Engine/Source/Runtime/TraceLog/create_standalone.py
@@ -3,6 +3,7 @@ import re
 import stat
 import shutil
 import argparse
+import subprocess
 from pathlib import Path
 
 #-------------------------------------------------------------------------------
@@ -19,6 +20,7 @@ class _SourceFile(object):
 		self.is_include = (path.suffix == ".h")
 		self.lines = []
 		self.deps = []
+		self.ext_deps = []
 
 	def __eq__(self, rhs):
 		return rhs.samefile(self.path)
@@ -64,16 +66,13 @@ def _finalize_source(source_files, include_dirs):
 					source_file.lines.append(line)
 				continue
 
-			is_local = False
-
 			for include_dir in (file.parent, *include_dirs):
 				candidate = include_dir / include
 				if candidate.is_file():
-					is_local = True
+					source_file.deps.append(candidate)
 					break
-
-			if is_local:
-				source_file.deps.append(candidate)
+			else:
+				source_file.ext_deps.append(include)
 
 	# Hook up dependencies
 	_spam("Resolving include dependencies")
@@ -155,63 +154,98 @@ def _main_trace(src_dir, dest_dir, thin, analysis):
 
 	# Add prologue and epilogue files
 	prologue = _SourceFile(src_dir / "standalone_prologue.h")
+	source_files.insert(0, prologue)
+	prologue.lines = [x for x in prologue.path.open("rt") if x]
+
 	epilogue = _SourceFile(src_dir / "standalone_epilogue.h")
-	if not thin:
-		source_files.insert(0, prologue)
-		source_files.append(epilogue)
-		prologue.lines = [x for x in prologue.path.open("rt") if x]
-		epilogue.lines = [x for x in epilogue.path.open("rt") if x]
+	source_files.append(epilogue)
+	epilogue.lines = [x for x in epilogue.path.open("rt") if x]
 
 	# Write the output file
 	_spam_header("Output")
 	_spam(dest_dir / "trace.h")
-	cpp_started = False
 	with (dest_dir / "trace.h").open("wt") as out:
 		print("// Copyright Epic Games, Inc. All Rights Reserved.", file=out)
 		print("#pragma once", file=out)
 
-		if thin:
-			include_path = str(prologue.path.resolve()).replace("\\", "/")
-			print(f'#include "{include_path}"', file=out)
-
-		for source_file in source_files:
-			if not cpp_started:
-				if source_file.path.suffix == ".cpp":
-					print("#if TRACE_IMPLEMENT", file=out)
-					cpp_started = True
-			elif source_file.path.suffix != ".cpp":
-				print("#endif // TRACE_IMPLEMENT", file=out)
-				cpp_started = False
-
-			print("/* {{{1", source_file.path.name, "*/", file=out)
-			print(file=out)
-			for line in source_file.lines:
-				out.write(line)
-
-		if thin:
-			include_path = str(epilogue.path.resolve()).replace("\\", "/")
+		if analysis_dir:
+			print("#define TRACE_HAS_ANALYSIS", file=out)
+
+		class State(object): pass
+		state = State()
+		state.out = out
+		state.dest_dir = dest_dir
+		state.src_dir = src_dir
+		state.analysis_dir = analysis_dir
+		state.source_files = source_files
+		return _thin(state) if thin else _fat(state)
+
+	_spam("...done!")
+
+#-------------------------------------------------------------------------------
+def _fat(state):
+	out = state.out
+	cpp_started = False
+	for source_file in state.source_files:
+		if not cpp_started:
+			if source_file.path.suffix == ".cpp":
+				print("#if TRACE_IMPLEMENT", file=out)
+				cpp_started = True
+		elif source_file.path.suffix != ".cpp":
 			print("#endif // TRACE_IMPLEMENT", file=out)
-			print(f'#include "{include_path}"', file=out)
-
-	# Copy some support files for thin iteration
-	if thin:
-		_spam_header("Support files for thin trace.h")
-		support_paths = (
-			src_dir / "Private/Trace/LZ4/lz4.h",
-			src_dir / "Private/Trace/LZ4/lz4.c.inl",
+			cpp_started = False
+
+		print("/* {{{1", source_file.path.name, "*/", file=out)
+		print(file=out)
+		for line in source_file.lines:
+			out.write(line)
+
+#-------------------------------------------------------------------------------
+def _thin(state):
+	# Include trace source code
+	def write_include(source_file):
+		path = str(source_file.path).replace("\\", "/")
+		print(fr'#include "{path}"', file=state.out)
+
+	out = state.out
+	write_include(state.source_files[0])
+	print("#if TRACE_IMPLEMENT", file=out)
+	for source_file in state.source_files[1:-1]:
+		if source_file.path.suffix == ".cpp":
+			write_include(source_file)
+	print("#endif // TRACE_IMPLEMENT", file=out)
+	write_include(state.source_files[-1])
+
+	# Stub out external includes
+	ext_deps = set()
+	for source_file in state.source_files:
+		for dep in source_file.ext_deps:
+			ext_deps.add(dep)
+
+	stubs_dir = state.dest_dir / "include"
+	for dep in ext_deps:
+		(stubs_dir / dep).parent.mkdir(parents=True, exist_ok=True)
+		(stubs_dir / dep).open("wt").close()
+
+	# Symlink source
+	def symlink(fr, to):
+		fr = state.dest_dir / "src" / fr
+		subprocess.run(
+			("cmd.exe", "/c", "mklink", "/j", fr, to.resolve()),
+			stderr=subprocess.DEVNULL
 		)
-		for support_path in support_paths:
-			dest_path = dest_dir / support_path.name
-			shutil.copy(support_path, dest_path)
-			os.chmod(dest_path, stat.S_IWRITE)
 
-	_spam("...done!")
+	(state.dest_dir / "src").mkdir(parents=True, exist_ok=True)
+	symlink("trace", state.src_dir)
+	if state.analysis_dir:
+		symlink("analysis", state.analysis_dir)
 
+#-------------------------------------------------------------------------------
 def main():
 	desc = "Amalgamate TraceLog into a standalone single-file library"
 	parser = argparse.ArgumentParser(description=desc)
 	parser.add_argument("outdir", help="Directory to write output file(s) to")
-	parser.add_argument("--thin", action="store_true", help="#include standalone_ instead of blitting")
+	parser.add_argument("--thin", action="store_true", help="#include everything instead of blitting")
 	parser.add_argument("--analysis", action="store_true", help="Amalgamate TraceAnalysis too")
 	args = parser.parse_args()
 
