Allow a small amount of tracing when allocating memory for trace buffers

Previously any events traced while Trace allocates memory to trace events into were dropped because there was nowhere to write these events. This change
splits memory allocation into two functions, one that will capture and send events when safe to do so, and one that doesn't capture events at all. For that latter, the capture is moved out to where allocating trace buffers happens so that captured events can be added to the buffer just allocated. This ensure order is maintained and that serials are ascending on a thread.

Partial buffers; prior to this change, aux-data that was larger than available space in a buffer would be fragmented across multiple buffers. Buffers that started in the middle of aux-data were marked partial as there was no way for analysis to interpret the stream from this point. As it is now possible that allocation-traced events will be written into the buffer before it's made available, each fragment is preceded with a aux-data header. This makes the stream recoverable from any buffer and negates the need for having partial buffers at all. As aux-data can now be constructed of two or more aux-data events for a single field of an event, analysis needs to be aware of this and defragment these events.

Late connect; analysis of a stream can now start with aux-data/terminal events with no corresponding owning event. This is fine; they get skipped over in DispatchEvents().

It is a deliberate choice to have a limited amount of redirected-buffer space to capture allocation-traced events. There needs to be enough space remaining in the new buffer to trace the outer event. Some platforms also have a limited amount of stack per thread.

The Writer_AddPageToPool() in Writer_InitializePool() didn't do anything and was removed.

GInitialized was set to early and subvert the reliance on static initialisation guarantees

#rnx

diff --git a/Engine/Source/Developer/TraceAnalysis/Private/Analysis/Engine.cpp b/Engine/Source/Developer/TraceAnalysis/Private/Analysis/Engine.cpp
--- a/Engine/Source/Developer/TraceAnalysis/Private/Analysis/Engine.cpp
+++ b/Engine/Source/Developer/TraceAnalysis/Private/Analysis/Engine.cpp
@@ -3,6 +3,7 @@
 #include "Engine.h"
 #include "Algo/BinarySearch.h"
 #include "Algo/Sort.h"
+#include "Algo/StableSort.h"
 #include "Containers/ArrayView.h"
 #include "CoreGlobals.h"
 #include "HAL/UnrealMemory.h"
@@ -53,15 +54,106 @@ struct FAuxData
 	uint32			DataSize;
 	uint16			FieldIndex;
 	int8			FieldSizeAndType;
-	bool			bSigned;
+	uint8			bSigned		: 1;
+	uint8			bFragmented	: 1;
+	uint8			bOwnsData	: 1;
+	uint8			_Unused		: 5;
 };
 
 ////////////////////////////////////////////////////////////////////////////////
 struct FAuxDataCollector
-: public TArray<FAuxData, TInlineAllocator<8>>
+	: public TArray<FAuxData, TInlineAllocator<8>>
 {
+	using	Super = TArray<FAuxData, TInlineAllocator<8>>;
+
+			~FAuxDataCollector() { Reset(); }
+	void	Add(const FAuxData& Data);
+	void	Defragment(FAuxData& Head);
+	void	Reset();
 };
 
+////////////////////////////////////////////////////////////////////////////////
+void FAuxDataCollector::Add(const FAuxData& Data)
+{
+	if (Num() == 0)
+	{
+		Super::Add(Data);
+		return;
+	}
+
+	FAuxData* Prev = &Last();
+	if (Prev->FieldIndex != Data.FieldIndex)
+	{
+		Super::Add(Data);
+		return;
+	}
+
+	while (Prev > GetData())
+	{
+		if (Prev[-1].FieldIndex != Data.FieldIndex)
+		{
+			break;
+		}
+
+		--Prev;
+	}
+
+	Prev->bFragmented = 1;
+	Super::Add(Data);
+}
+
+////////////////////////////////////////////////////////////////////////////////
+void FAuxDataCollector::Defragment(FAuxData& Data)
+{
+	check(Data.bFragmented);
+
+	uint32 Size = Data.DataSize;
+	for (FAuxData* Read = &Data + 1;;)
+	{
+		Size += Read->DataSize;
+
+		++Read;
+		if (Read > &Last() || Read->FieldIndex != Data.FieldIndex)
+		{
+			break;
+		}
+	}
+
+	uint8* Buffer = (uint8*)(FMemory::Malloc(Size));
+
+	uint8* Write = Buffer;
+	uint8* End = Buffer + Size;
+	for (FAuxData* Read = &Data;; ++Read)
+	{
+		FMemory::Memcpy(Write, Read->Data, Read->DataSize);
+		Write += Read->DataSize;
+		if (Write >= End)
+		{
+			break;
+		}
+	}
+
+	Data.Data = Buffer;
+	Data.DataSize = Size;
+	Data.bOwnsData = 1;
+}
+
+////////////////////////////////////////////////////////////////////////////////
+void FAuxDataCollector::Reset()
+{
+	for (const FAuxData& AuxData : *this)
+	{
+		if (!AuxData.bOwnsData)
+		{
+			continue;
+		}
+
+		FMemory::Free((void*)(AuxData.Data));
+	}
+
+	Super::Reset();
+}
+
 
 
 // {{{1 threads ----------------------------------------------------------------
@@ -657,6 +749,11 @@ const FAuxData* FEventDataInfo::GetAuxData(uint32 FieldIndex) const
 	{
 		if (Data.FieldIndex == FieldIndex)
 		{
+			if (Data.bFragmented)
+			{
+				Info->AuxCollector->Defragment(Data);
+			}
+
 			const FDispatch::FField& Field = Info->Dispatch.Fields[FieldIndex];
 			Data.FieldSizeAndType = Field.SizeAndType;
 			Data.bSigned = (Field.Class & UE::Trace::Protocol0::Field_Signed) != 0;
@@ -2455,7 +2552,7 @@ int32 FProtocol2Stage::OnDataAux(FStreamReader& Reader, FAuxDataCollector& Colle
 		}
 
 		// Attach to event
-		FAuxData AuxData;
+		FAuxData AuxData = {};
 		AuxData.Data = Header->Data;
 		AuxData.DataSize = uint32(BlockSize - sizeof(*Header));
 		AuxData.FieldIndex = uint16(Header->FieldIndex & Protocol1::FAuxHeader::FieldMask);
@@ -3298,10 +3395,10 @@ int32 FProtocol5Stage::ParseEventsWithAux(FStreamReader& Reader, EventDescArray&
 			if (EventDesc.bHasAux)
 			{
 				AuxKeyStack.Add(uint16(AuxKey));
-				bUnsorted = true;
-
-				++AuxKey; // Add a gap for aux-blocks to sort into
 			}
+
+			// This event maybe in the middle of an earlier event's aux data blocks.
+			bUnsorted = true;
 		}
 
 		OutEventDescs.Add(EventDesc);
@@ -3322,7 +3419,7 @@ int32 FProtocol5Stage::ParseEventsWithAux(FStreamReader& Reader, EventDescArray&
 	{
 		uint32 NumDescs = OutEventDescs.Num() - FirstDescIndex;
 		TArrayView<FEventDesc> DescsView(OutEventDescs.GetData() + FirstDescIndex, NumDescs);
-		Algo::Sort(
+		Algo::StableSort(
 			DescsView,
 			[] (const FEventDesc& Lhs, const FEventDesc& Rhs)
 			{
@@ -3641,7 +3738,7 @@ int32 FProtocol5Stage::DispatchEvents(
 
 				const auto* AuxHeader = ((FAuxHeader*)(Cursor->Data)) - 1;
 
-				FAuxData AuxData;
+				FAuxData AuxData = {};
 				AuxData.Data = AuxHeader->Data;
 				AuxData.DataSize = (AuxHeader->Pack >> FAuxHeader::SizeShift);
 				AuxData.FieldIndex = AuxHeader->FieldIndex_Size & FAuxHeader::FieldMask;
diff --git a/Engine/Source/Runtime/TraceLog/Private/Trace/BlockPool.cpp b/Engine/Source/Runtime/TraceLog/Private/Trace/BlockPool.cpp
--- a/Engine/Source/Runtime/TraceLog/Private/Trace/BlockPool.cpp
+++ b/Engine/Source/Runtime/TraceLog/Private/Trace/BlockPool.cpp
@@ -13,7 +13,7 @@ namespace Trace {
 namespace Private {
 
 ////////////////////////////////////////////////////////////////////////////////
-void*		Writer_MemoryAllocate(SIZE_T, uint32);
+void*		Writer_MemoryAllocateNoRedirect(SIZE_T, uint32);
 void		Writer_MemoryFree(void*, uint32);
 
 
@@ -49,7 +49,7 @@ static uint32							GPoolUsage;			// = 0;
 static FPoolBlockList Writer_AddPageToPool(uint32 PageSize)
 {
 	// The free list is empty so we have to populate it with some new blocks.
-	uint8* PageBase = (uint8*)Writer_MemoryAllocate(PageSize, PLATFORM_CACHE_LINE_SIZE);
+	uint8* PageBase = (uint8*)Writer_MemoryAllocateNoRedirect(PageSize, PLATFORM_CACHE_LINE_SIZE);
 	GPoolUsage += PageSize;
 
 	uint32 BufferSize = GPoolBlockSize;
@@ -162,7 +162,6 @@ void Writer_FreeBlockListToPool(FWriteBuffer* Head, FWriteBuffer* Tail)
 ////////////////////////////////////////////////////////////////////////////////
 void Writer_InitializePool()
 {
-	Writer_AddPageToPool(GPoolBlockSize);
 	static_assert(GPoolPageSize >= 0x10000, "Page growth must be >= 64KB");
 	static_assert(GPoolInitPageSize >= 0x10000, "Initial page size must be >= 64KB");
 }
diff --git a/Engine/Source/Runtime/TraceLog/Private/Trace/Field.cpp b/Engine/Source/Runtime/TraceLog/Private/Trace/Field.cpp
--- a/Engine/Source/Runtime/TraceLog/Private/Trace/Field.cpp
+++ b/Engine/Source/Runtime/TraceLog/Private/Trace/Field.cpp
@@ -15,7 +15,8 @@ template <typename CallbackType>
 static void Field_WriteAuxData(uint32 Index, int32 Size, CallbackType&& Callback)
 {
 	static_assert(
-		sizeof(Private::FWriteBuffer::Overflow) >= sizeof(FAuxHeader) + sizeof(uint8 /*AuxDataTerminal*/),
+		// "1+" is so we've at least space to write one byte of payload
+		sizeof(Private::FWriteBuffer::Overflow) >= 1 + sizeof(FAuxHeader) + sizeof(uint8 /*AuxDataTerminal*/),
 		"FWriteBuffer::Overflow is not large enough"
 	);
 
@@ -25,41 +26,59 @@ static void Field_WriteAuxData(uint32 Index, int32 Size, CallbackType&& Callback
 		return;
 	}
 
-	// Header
 	FWriteBuffer* Buffer = Writer_GetBuffer();
-	auto* Header = (FAuxHeader*)(Buffer->Cursor);
-	Header->Pack = Size << FAuxHeader::SizeShift;
-	Header->Pack |= Index << FAuxHeader::FieldShift;
-	Header->Uid = uint8(EKnownEventUids::AuxData) << EKnownEventUids::_UidShift;
-	Buffer->Cursor += sizeof(FAuxHeader);
-
-	// Array data
-	bool bCommit = ((uint8*)Header == Buffer->Committed);
-	while (true)
+
+	// We are writing fields of an event. The event writing will move Commited
+	// along. However, if we're writing a second aux field where the first fetched
+	// a new buffer then it is a different buffer to the one the event's written
+	// into, thus it is us that needs to update Committed.
+	bool bCommit = (Buffer->Cursor == Buffer->Committed);
+
+	// Do we have enough space to write at least [header][overflow-1] bytes of payload?
+	int32 Remaining = int32(ptrdiff_t((uint8*)Buffer - Buffer->Cursor));
+	auto NextBuffer = [&Buffer, &Remaining, &bCommit] ()
 	{
-		if (Buffer->Cursor >= (uint8*)Buffer)
+		if (bCommit)
 		{
-			if (bCommit)
-			{
-				AtomicStoreRelease(&(uint8* volatile&)(Buffer->Committed), Buffer->Cursor);
-			}
+			AtomicStoreRelease(&(uint8* volatile&)(Buffer->Committed), Buffer->Cursor);
+		}
+		bCommit = true;
 
-			Buffer = Writer_NextBuffer(0);
-			Buffer->Partial = 1;
+		Buffer = Writer_NextBuffer();
+		Remaining = int32(ptrdiff_t((uint8*)Buffer - Buffer->Cursor));
+	};
 
-			bCommit = true;
-		}
+	if (Remaining <= 0)
+	{
+		NextBuffer();
+	}
 
-		int32 Remaining = int32((uint8*)Buffer - Buffer->Cursor);
+	while (true)
+	{
+		// -1 to leave space to later write aux-term without bounds checking
+		Remaining += sizeof(FWriteBuffer::Overflow) - 1;
 		int32 SegmentSize = (Remaining < Size) ? Remaining : Size;
+
+		// Write header
+		auto* Header = (FAuxHeader*)(Buffer->Cursor);
+		Header->Pack = SegmentSize << FAuxHeader::SizeShift;
+		Header->Pack |= Index << FAuxHeader::FieldShift;
+		Header->Uid = uint8(EKnownEventUids::AuxData) << EKnownEventUids::_UidShift;
+		Buffer->Cursor += sizeof(FAuxHeader);
+		Remaining -= sizeof(FAuxHeader);
+
+		// Write payload data
 		Callback(Buffer->Cursor, SegmentSize);
 		Buffer->Cursor += SegmentSize;
 
+		// Bounds check
 		Size -= SegmentSize;
 		if (Size <= 0)
 		{
 			break;
 		}
+
+		NextBuffer();
 	}
 
 	if (bCommit)
diff --git a/Engine/Source/Runtime/TraceLog/Private/Trace/Tail.cpp b/Engine/Source/Runtime/TraceLog/Private/Trace/Tail.cpp
--- a/Engine/Source/Runtime/TraceLog/Private/Trace/Tail.cpp
+++ b/Engine/Source/Runtime/TraceLog/Private/Trace/Tail.cpp
@@ -199,7 +199,7 @@ FTidPacketBase* FPacketRing::AppendImpl(uint32 InSize)
 static FPacketRing GPacketRing; // = {};
 
 ////////////////////////////////////////////////////////////////////////////////
-void Writer_TailAppend(uint32 ThreadId, uint8* __restrict Data, uint32 Size, bool bPartial)
+void Writer_TailAppend(uint32 ThreadId, uint8* __restrict Data, uint32 Size)
 {
 	// Perhaps tail tracing is disabled?
 	if (!GPacketRing.IsActive())
@@ -216,7 +216,6 @@ void Writer_TailAppend(uint32 ThreadId, uint8* __restrict Data, uint32 Size, boo
 	}
 
 	ThreadId &= FTidPacketBase::ThreadIdMask;
-	ThreadId |= bPartial ? FTidPacketBase::PartialMarker : 0;
 
 	// Smaller buffers usually aren't redundant enough to benefit from being
 	// compressed. They often end up being larger.
@@ -310,10 +309,6 @@ static void StressRingPacket()
 	{
 		FTidPacket* Packet = Ring.Append<FTidPacket>((Bits & 0x1f) + 6);
 		Packet->ThreadId = i;
-		if ((Bits & 0x15) == 0)
-		{
-			Packet->ThreadId |= FTidPacketBase::PartialMarker;
-		}
 
 		Ring.IterateRanges([] (const FPacketRing::FRange&)
 		{
diff --git a/Engine/Source/Runtime/TraceLog/Private/Trace/TlsBuffer.cpp b/Engine/Source/Runtime/TraceLog/Private/Trace/TlsBuffer.cpp
--- a/Engine/Source/Runtime/TraceLog/Private/Trace/TlsBuffer.cpp
+++ b/Engine/Source/Runtime/TraceLog/Private/Trace/TlsBuffer.cpp
@@ -4,17 +4,19 @@
 
 #if UE_TRACE_ENABLED
 
+#include "HAL/Platform.h" // for PLATFORM_BREAK
 #include "Platform.h"
 #include "Trace/Detail/Atomic.h"
 #include "Trace/Detail/Writer.inl"
 #include "Trace/Trace.inl"
+#include "WriteBufferRedirect.h"
 
 namespace UE {
 namespace Trace {
 namespace Private {
 
 ////////////////////////////////////////////////////////////////////////////////
-void				Writer_TailAppend(uint32, uint8* __restrict, uint32, bool);
+void				Writer_TailAppend(uint32, uint8* __restrict, uint32);
 FWriteBuffer*		Writer_AllocateBlockFromPool();
 uint32				Writer_GetThreadId();
 void				Writer_FreeBlockListToPool(FWriteBuffer*, FWriteBuffer*);
@@ -32,7 +34,7 @@ UE_TRACE_EVENT_END()
 
 ////////////////////////////////////////////////////////////////////////////////
 #define T_ALIGN alignas(PLATFORM_CACHE_LINE_SIZE)
-static FWriteBuffer						GNullWriteBuffer	= { {}, 0, 0, nullptr, nullptr, (uint8*)&GNullWriteBuffer };
+static FWriteBuffer						GNullWriteBuffer	= { {}, 0, nullptr, nullptr, (uint8*)&GNullWriteBuffer, nullptr, nullptr, 0, 0, 0 };
 thread_local FWriteBuffer*				GTlsWriteBuffer		= &GNullWriteBuffer;
 static FWriteBuffer* __restrict			GActiveThreadList;	// = nullptr;
 T_ALIGN static FWriteBuffer* volatile	GNewThreadList;		// = nullptr;
@@ -49,22 +51,46 @@ TRACELOG_API FWriteBuffer* Writer_GetBuffer()
 #endif
 
 ////////////////////////////////////////////////////////////////////////////////
-static FWriteBuffer* Writer_NextBufferInternal()
+static FWriteBuffer* Writer_NextBufferInternal(FWriteBuffer* CurrentBuffer)
 {
-	FWriteBuffer* NextBuffer = Writer_AllocateBlockFromPool();
+	// TraceData is used to catch traced events that occur when allocating
+	// memory. Such events should be tightly controlled and kept to a minimum.
+	// There is not a lot of buffer space available in TraceData to avoid
+	// exhausting available stack and leave usable space in the new buffer.
+	TWriteBufferRedirect<320> TraceData;
+
+#if !UE_BUILD_SHIPPING && !UE_BUILD_TEST
+	if (CurrentBuffer->ThreadId == decltype(TraceData)::ActiveRedirection)
+	{
+		// If we've reached this point we're are in an unrecoverable situation.
+		// Allocating memory results in so many trace events that there is
+		// insufficient space to store them. We can't allocate more space, because
+		// that would result in more traced events, and so on...
+		PLATFORM_BREAK();
+	}
+#endif
 
+	FWriteBuffer* NextBuffer = Writer_AllocateBlockFromPool();
 	NextBuffer->Cursor = (uint8*)NextBuffer - NextBuffer->Size;
-	NextBuffer->Committed = NextBuffer->Cursor;
 	NextBuffer->Reaped = NextBuffer->Cursor;
 	NextBuffer->EtxOffset = 0 - int32(sizeof(FWriteBuffer));
 	NextBuffer->NextBuffer = nullptr;
 
-	FWriteBuffer* CurrentBuffer = GTlsWriteBuffer;
+	// Add any capture events into the buffer.
+	if (uint32 RedirectSize = TraceData.GetSize())
+	{
+		memcpy(NextBuffer->Cursor, TraceData.GetData(), RedirectSize);
+		NextBuffer->Cursor += RedirectSize;
+	}
+
+	TraceData.Abandon();
+
+	NextBuffer->Committed = NextBuffer->Cursor;
+
 	if (CurrentBuffer == &GNullWriteBuffer)
 	{
 		NextBuffer->ThreadId = uint16(Writer_GetThreadId());
 		NextBuffer->PrevTimestamp = TimeGetTimestamp();
-		NextBuffer->Partial = 0;
 
 		GTlsWriteBuffer = NextBuffer;
 
@@ -86,7 +112,6 @@ static FWriteBuffer* Writer_NextBufferInternal()
 		CurrentBuffer->NextBuffer = NextBuffer;
 		NextBuffer->ThreadId = CurrentBuffer->ThreadId;
 		NextBuffer->PrevTimestamp = CurrentBuffer->PrevTimestamp;
-		NextBuffer->Partial = 0;
 
 		GTlsWriteBuffer = NextBuffer;
 
@@ -99,24 +124,10 @@ static FWriteBuffer* Writer_NextBufferInternal()
 }
 
 ////////////////////////////////////////////////////////////////////////////////
-TRACELOG_API FWriteBuffer* Writer_NextBuffer(int32 Size)
+TRACELOG_API FWriteBuffer* Writer_NextBuffer()
 {
 	FWriteBuffer* CurrentBuffer = GTlsWriteBuffer;
-	if (CurrentBuffer != &GNullWriteBuffer)
-	{
-		CurrentBuffer->Cursor -= Size;
-	}
-
-	FWriteBuffer* NextBuffer = Writer_NextBufferInternal();
-
-	if (Size >= NextBuffer->Size)
-	{
-		// Someone is trying to write an event that is far too large
-		return nullptr;
-	}
-
-	NextBuffer->Cursor += Size;
-	return NextBuffer;
+	return Writer_NextBufferInternal(CurrentBuffer);
 }
 
 ////////////////////////////////////////////////////////////////////////////////
@@ -131,9 +142,7 @@ static bool Writer_DrainBuffer(uint32 ThreadId, FWriteBuffer* Buffer)
 		GTraceStatistics.BytesTraced += SizeToReap;
 #endif
 
-		bool bPartial = (Buffer->Partial == 1);
-		bPartial &= UPTRINT(Buffer->Reaped + Buffer->Size) == UPTRINT(Buffer);
-		Writer_TailAppend(ThreadId, Buffer->Reaped, SizeToReap, bPartial);
+		Writer_TailAppend(ThreadId, Buffer->Reaped, SizeToReap);
 		Buffer->Reaped = Committed;
 	}
 
diff --git a/Engine/Source/Runtime/TraceLog/Private/Trace/WriteBufferRedirect.h b/Engine/Source/Runtime/TraceLog/Private/Trace/WriteBufferRedirect.h
--- a/Engine/Source/Runtime/TraceLog/Private/Trace/WriteBufferRedirect.h
+++ b/Engine/Source/Runtime/TraceLog/Private/Trace/WriteBufferRedirect.h
@@ -16,9 +16,12 @@ template <int BufferSize>
 class TWriteBufferRedirect
 {
 public:
+	enum : uint16 { ActiveRedirection = 0xffff };
+
 					TWriteBufferRedirect();
 					~TWriteBufferRedirect();
 	void			Close();
+	void			Abandon();
 	uint8*			GetData();
 	uint32			GetSize() const;
 	uint32			GetCapacity() const;
@@ -37,6 +40,8 @@ inline TWriteBufferRedirect<BufferSize>::TWriteBufferRedirect()
 	Reset();
 	PrevBuffer = GTlsWriteBuffer;
 	GTlsWriteBuffer = &Buffer;
+	Buffer.Size = uint16(BufferSize);
+	Buffer.ThreadId = ActiveRedirection;
 }
 
 ////////////////////////////////////////////////////////////////////////////////
@@ -59,6 +64,13 @@ inline void TWriteBufferRedirect<BufferSize>::Close()
 	PrevBuffer = nullptr;
 }
 
+////////////////////////////////////////////////////////////////////////////////
+template <int BufferSize>
+inline void TWriteBufferRedirect<BufferSize>::Abandon()
+{
+	PrevBuffer = nullptr;
+}
+
 ////////////////////////////////////////////////////////////////////////////////
 template <int BufferSize>
 inline uint8* TWriteBufferRedirect<BufferSize>::GetData()
diff --git a/Engine/Source/Runtime/TraceLog/Private/Trace/Writer.cpp b/Engine/Source/Runtime/TraceLog/Private/Trace/Writer.cpp
--- a/Engine/Source/Runtime/TraceLog/Private/Trace/Writer.cpp
+++ b/Engine/Source/Runtime/TraceLog/Private/Trace/Writer.cpp
@@ -138,10 +138,8 @@ void Writer_MemorySetHooks(decltype(AllocHook) Alloc, decltype(FreeHook) Free)
 }
 
 ////////////////////////////////////////////////////////////////////////////////
-void* Writer_MemoryAllocate(SIZE_T Size, uint32 Alignment)
+void* Writer_MemoryAllocateNoRedirect(SIZE_T Size, uint32 Alignment)
 {
-	TWriteBufferRedirect<1 << 10> TraceData;
-
 	void* Ret = nullptr;
 
 #if TRACE_PRIVATE_STOMP
@@ -186,9 +184,27 @@ void* Writer_MemoryAllocate(SIZE_T Size, uint32 Alignment)
 	return Ret;
 }
 
+////////////////////////////////////////////////////////////////////////////////
+void* Writer_MemoryAllocate(SIZE_T Size, uint32 Alignment)
+{
+	TWriteBufferRedirect<1 << 10> TraceData;
+
+	void* Ret = Writer_MemoryAllocateNoRedirect(Size, Alignment);
+
+	if (TraceData.GetSize())
+	{
+		uint32 ThreadId = Writer_GetThreadId();
+		Writer_SendData(ThreadId, TraceData.GetData(), TraceData.GetSize());
+	}
+
+	return Ret;
+}
+
 ////////////////////////////////////////////////////////////////////////////////
 void Writer_MemoryFree(void* Address, uint32 Size)
 {
+	TWriteBufferRedirect<1 << 10> TraceData;
+
 #if TRACE_PRIVATE_STOMP
 	if (Address == nullptr)
 	{
@@ -203,8 +219,6 @@ void Writer_MemoryFree(void* Address, uint32 Size)
 	DWORD Unused;
 	VirtualProtect(MemInfo.BaseAddress, MemInfo.RegionSize, PAGE_READONLY, &Unused);
 #else // TRACE_PRIVATE_STOMP
-	TWriteBufferRedirect<1 << 10> TraceData;
-
 	if (FreeHook != nullptr)
 	{
 		FreeHook(Address, Size);
@@ -222,6 +236,12 @@ void Writer_MemoryFree(void* Address, uint32 Size)
 #if TRACE_PRIVATE_STATISTICS
 	AtomicAddRelaxed(&GTraceStatistics.MemoryUsed, uint64(-int64(Size)));
 #endif
+
+	if (TraceData.GetSize())
+	{
+		uint32 ThreadId = Writer_GetThreadId();
+		Writer_SendData(ThreadId, TraceData.GetData(), TraceData.GetSize());
+	}
 }
 
 
@@ -653,13 +673,14 @@ static void Writer_InternalInitializeImpl()
 		return;
 	}
 
-	GInitialized = true;
 	GStartCycle = TimeGetTimestamp();
 
 	Writer_InitializeSharedBuffers();
 	Writer_InitializePool();
 	Writer_InitializeControl();
 
+	GInitialized = true;
+
 	UE_TRACE_LOG($Trace, NewTrace, TraceLogChannel)
 		<< NewTrace.StartCycle(GStartCycle)
 		<< NewTrace.CycleFrequency(TimeGetFrequency())
diff --git a/Engine/Source/Runtime/TraceLog/Public/Trace/Detail/LogScope.inl b/Engine/Source/Runtime/TraceLog/Public/Trace/Detail/LogScope.inl
--- a/Engine/Source/Runtime/TraceLog/Public/Trace/Detail/LogScope.inl
+++ b/Engine/Source/Runtime/TraceLog/Public/Trace/Detail/LogScope.inl
@@ -8,6 +8,7 @@
 
 #include "Atomic.h"
 #include "EventNode.h"
+#include "HAL/Platform.h" // for PLATFORM_BREAK
 #include "LogScope.h"
 #include "Protocol.h"
 #include "Writer.inl"
@@ -59,13 +60,23 @@ inline void FLogScope::EnterPrelude(uint32 Size)
 	uint32 AllocSize = sizeof(HeaderType) + Size;
 
 	Buffer = Writer_GetBuffer();
-	Buffer->Cursor += AllocSize;
-	if (UNLIKELY(Buffer->Cursor > (uint8*)Buffer))
+	if (UNLIKELY(Buffer->Cursor + AllocSize > (uint8*)Buffer))
 	{
-		Buffer = Writer_NextBuffer(AllocSize);
+		Buffer = Writer_NextBuffer();
 	}
 
-	Ptr = Buffer->Cursor - Size;
+#if !UE_BUILD_SHIPPING && !UE_BUILD_TEST
+	if (AllocSize >= Buffer->Size)
+	{
+		// This situation is terminal. Someone's trying to trace an event that
+		// is far too large. This should never happen as 'maximum_field_count *
+		// largest_field_size' won't exceed Buffer.Size.
+		PLATFORM_BREAK();
+	}
+#endif
+
+	Ptr = Buffer->Cursor + sizeof(HeaderType);
+	Buffer->Cursor += AllocSize;
 }
 
 ////////////////////////////////////////////////////////////////////////////////
@@ -125,7 +136,7 @@ inline FScopedLogScope::~FScopedLogScope()
 	FWriteBuffer* Buffer = Writer_GetBuffer();
 	if (UNLIKELY(int32((uint8*)Buffer - Buffer->Cursor) < int32(sizeof(LeaveUid))))
 	{
-		Buffer = Writer_NextBuffer(0);
+		Buffer = Writer_NextBuffer();
 	}
 
 	Buffer->Cursor[0] = LeaveUid;
@@ -156,7 +167,7 @@ inline FScopedStampedLogScope::~FScopedStampedLogScope()
 
 	if (UNLIKELY(int32((uint8*)Buffer - Buffer->Cursor) < int32(sizeof(Stamp))))
 	{
-		Buffer = Writer_NextBuffer(0);
+		Buffer = Writer_NextBuffer();
 	}
 
 	Stamp <<= 8;
@@ -193,7 +204,7 @@ FORCENOINLINE auto FLogScope::ScopedEnter()
 	FWriteBuffer* Buffer = Writer_GetBuffer();
 	if (UNLIKELY(int32((uint8*)Buffer - Buffer->Cursor) < int32(sizeof(EnterUid))))
 	{
-		Buffer = Writer_NextBuffer(0);
+		Buffer = Writer_NextBuffer();
 	}
 
 	Buffer->Cursor[0] = EnterUid;
@@ -213,7 +224,7 @@ FORCENOINLINE auto FLogScope::ScopedStampedEnter()
 	FWriteBuffer* Buffer = Writer_GetBuffer();
 	if (UNLIKELY(int32((uint8*)Buffer - Buffer->Cursor) < int32(sizeof(Stamp))))
 	{
-		Buffer = Writer_NextBuffer(0);
+		Buffer = Writer_NextBuffer();
 	}
 
 	Stamp = Writer_GetTimestamp(Buffer);
diff --git a/Engine/Source/Runtime/TraceLog/Public/Trace/Detail/Transport.h b/Engine/Source/Runtime/TraceLog/Public/Trace/Detail/Transport.h
--- a/Engine/Source/Runtime/TraceLog/Public/Trace/Detail/Transport.h
+++ b/Engine/Source/Runtime/TraceLog/Public/Trace/Detail/Transport.h
@@ -39,7 +39,7 @@ struct FTidPacketBase
 	enum : uint16
 	{
 		EncodedMarker = 0x8000,
-		PartialMarker = 0x4000,
+		PartialMarker = 0x4000, // now unused. fragmented aux-data has an event header
 		ThreadIdMask  = PartialMarker - 1,
 	};
 
diff --git a/Engine/Source/Runtime/TraceLog/Public/Trace/Detail/Writer.inl b/Engine/Source/Runtime/TraceLog/Public/Trace/Detail/Writer.inl
--- a/Engine/Source/Runtime/TraceLog/Public/Trace/Detail/Writer.inl
+++ b/Engine/Source/Runtime/TraceLog/Public/Trace/Detail/Writer.inl
@@ -11,8 +11,7 @@ namespace Private {
 ////////////////////////////////////////////////////////////////////////////////
 struct FWriteBuffer
 {
-	uint8						Overflow[6];
-	uint16						Size;
+	uint8						Overflow[8];
 	uint64						PrevTimestamp;
 	FWriteBuffer* __restrict	NextThread;
 	FWriteBuffer* __restrict	NextBuffer;
@@ -20,13 +19,13 @@ struct FWriteBuffer
 	uint8* __restrict volatile	Committed;
 	uint8* __restrict			Reaped;
 	int32 volatile				EtxOffset;
-	int16						Partial;
+	uint16						Size;
 	uint16						ThreadId;
 };
 
 ////////////////////////////////////////////////////////////////////////////////
 TRACELOG_API uint64				TimeGetTimestamp();
-TRACELOG_API FWriteBuffer*		Writer_NextBuffer(int32);
+TRACELOG_API FWriteBuffer*		Writer_NextBuffer();
 TRACELOG_API FWriteBuffer*		Writer_GetBuffer();
 
 ////////////////////////////////////////////////////////////////////////////////
