Updated standalone prologue to support latest Trace edits and compilation with a C++20 toolchain.

#rnx

diff --git a/Engine/Source/Runtime/TraceLog/standalone_prologue.h b/Engine/Source/Runtime/TraceLog/standalone_prologue.h
--- a/Engine/Source/Runtime/TraceLog/standalone_prologue.h
+++ b/Engine/Source/Runtime/TraceLog/standalone_prologue.h
@@ -33,9 +33,6 @@
 #	define PLATFORM_APPLE			1
 #endif
 
-#if PLATFORM_MAC
-#endif
-
 // arch defines
 #if defined(__amd64__) || defined(_M_X64)
 #	define PLATFORM_CPU_X86_FAMILY	1
@@ -52,6 +49,10 @@
 #endif
 
 // external includes
+#if defined(_MSC_VER)
+#	define _ENABLE_EXTENDED_ALIGNED_STORAGE
+#endif
+
 #include <algorithm>
 #include <cstddef>
 #include <cstdint>
@@ -143,12 +144,30 @@ using WIDECHAR	= char16_t;
 #define TRACE_PRIVATE_CONTROL_ENABLED	0
 #define TRACE_PRIVATE_EXTERNAL_LZ4		1
 #define UE_TRACE_ENABLED				TRACE_ENABLED
-#define UE_BUILD_SHIPPING				1
+#define PREPROCESSOR_JOIN(a,b)			PREPROCESSOR_JOIN_IMPL(a,b)
+#define PREPROCESSOR_JOIN_IMPL(a,b)		a##b
 #define DEFINE_LOG_CATEGORY_STATIC(...)
 #define UE_LOG(...)
 #define check(x)						do { if (!(x)) __debugbreak(); } while (0)
 #define checkf(x, ...)					check(x)
 
+#if !defined(UE_BUILD_SHIPPING)
+#	define UE_BUILD_SHIPPING			1
+#endif
+#define UE_BUILD_TEST					0
+
+#if defined(__has_builtin)
+#	if __has_builtin(__builtin_trap)
+#		define PLATFORM_BREAK()			do { __builtin_trap(); } while (0)
+#	endif
+#elif defined(_MSC_VER)
+#	define PLATFORM_BREAK()				do { __debugbreak(); } while (0)
+#endif
+
+#if !defined(PLATFORM_BREAK)
+#	define PLATFORM_BREAK()				do { *(int*)0x493 = 0x493; } while (0)
+#endif
+
 // api
 template <typename T> inline auto Forward(T t)		{ return std::forward<T>(t); }
 template <typename T> inline auto MoveTemp(T&& t)	{ return std::move<T>(t); }
@@ -160,7 +179,23 @@ template <typename T, typename ALLOCATOR=void>
 struct TArray
 	: public std::vector<T>
 {
-	using	 std::vector<T>::vector;
+	using	 Super = std::vector<T>;
+	using	 Super::vector;
+	using	 Super::back;
+	using	 Super::begin;
+	using	 Super::clear;
+	using	 Super::data;
+	using	 Super::emplace_back;
+	using	 Super::empty;
+	using	 Super::end;
+	using	 Super::erase;
+	using	 Super::front;
+	using	 Super::insert;
+	using	 Super::pop_back;
+	using	 Super::push_back;
+	using	 Super::reserve;
+	using	 Super::resize;
+	using	 Super::size;
 	void	 Add(const T& t)					{ push_back(t); }
 	T&		 Add_GetRef(const T& t)				{ push_back(t); return back(); }
 	size_t	 AddUninitialized(size_t n)			{ resize(size() + n); return size() - n; }
@@ -214,7 +249,10 @@ template <typename T>
 struct TStringViewAdapter
 	: public std::basic_string_view<T>
 {
-	using		std::basic_string_view<T>::basic_string_view;
+	using		Super = std::basic_string_view<T>;
+	using		Super::basic_string_view;
+	using		Super::size;
+	using		Super::data;
 	size_t		Len() const		{ return size(); }
 	const char*	GetData() const	{ return data(); }
 };
@@ -272,6 +310,13 @@ inline void*	Realloc(void* ptr, size_t size)			{ return ::realloc(ptr, size); }
 inline void		Free(void* ptr)							{ return ::free(ptr); }
 inline void		Memcpy(void* d, const void* s, size_t n){ ::memcpy(d, s, n); }
 
+template <typename T, typename COMPARE=std::less<T::value_type>>
+inline void StableSort(T& container, COMPARE&& compare=COMPARE())
+{
+	std::stable_sort(container.GetData(), container.GetData() + container.Num(),
+		std::forward<COMPARE>(compare));
+}
+
 template <typename T, typename COMPARE=std::less<T::value_type>>
 inline void Sort(T& container, COMPARE&& compare=COMPARE())
 {
