Fixed stomping of PrevTimestamp's lower bits.

The size of the segment of aux-data to memcpy into the buffer incorrectly included the size of the aux-header. This would result in writing aux-data off the end of the overflow and into the low-order bits of the timestamp bias. This would inadvertantly cause incorrect, sometimes negative, timestamps for scoped events.

#rnx

diff --git a/Engine/Source/Runtime/TraceLog/Private/Trace/Field.cpp b/Engine/Source/Runtime/TraceLog/Private/Trace/Field.cpp
--- a/Engine/Source/Runtime/TraceLog/Private/Trace/Field.cpp
+++ b/Engine/Source/Runtime/TraceLog/Private/Trace/Field.cpp
@@ -55,8 +55,11 @@ static void Field_WriteAuxData(uint32 Index, int32 Size, CallbackType&& Callback
 
 	while (true)
 	{
-		// -1 to leave space to later write aux-term without bounds checking
-		Remaining += sizeof(FWriteBuffer::Overflow) - 1;
+		// Buffers have a small overflow which we can use. It means we can write
+		// some elements unconditionally. What remains is free for payload data.
+		Remaining += sizeof(FWriteBuffer::Overflow);
+		Remaining -= 1;						// for the aux-terminal
+		Remaining -= sizeof(FAuxHeader);	// header also assume to always fit
 		int32 SegmentSize = (Remaining < Size) ? Remaining : Size;
 
 		// Write header
@@ -65,7 +68,7 @@ static void Field_WriteAuxData(uint32 Index, int32 Size, CallbackType&& Callback
 		Header->Pack |= Index << FAuxHeader::FieldShift;
 		Header->Uid = uint8(EKnownEventUids::AuxData) << EKnownEventUids::_UidShift;
 		Buffer->Cursor += sizeof(FAuxHeader);
-		Remaining -= sizeof(FAuxHeader);
+
 
 		// Write payload data
 		Callback(Buffer->Cursor, SegmentSize);
