Fixed runtimes not starting up with certain combinations of trace channels enabled

The redirect buffer was a little too small to contain the callstack that may get traced. Bumping the size isn't the ideal solution but suitable for the now. Alternatives would be to delta compress callstack events (only an option on monolithic builds), take smaller backtraces for trace-related allocations, eliminate redundancy in addresses (i.e. a 32 bit encoding scheme), or possibly segmented deduplication.

#rnx

diff --git a/Engine/Source/Runtime/TraceLog/Private/Trace/TlsBuffer.cpp b/Engine/Source/Runtime/TraceLog/Private/Trace/TlsBuffer.cpp
--- a/Engine/Source/Runtime/TraceLog/Private/Trace/TlsBuffer.cpp
+++ b/Engine/Source/Runtime/TraceLog/Private/Trace/TlsBuffer.cpp
@@ -57,7 +57,7 @@ static FWriteBuffer* Writer_NextBufferInternal(FWriteBuffer* CurrentBuffer)
 	// memory. Such events should be tightly controlled and kept to a minimum.
 	// There is not a lot of buffer space available in TraceData to avoid
 	// exhausting available stack and leave usable space in the new buffer.
-	TWriteBufferRedirect<320> TraceData;
+	TWriteBufferRedirect<512> TraceData;
 
 #if !UE_BUILD_SHIPPING && !UE_BUILD_TEST
 	if (CurrentBuffer->ThreadId == decltype(TraceData)::ActiveRedirection)
