Removed event redirection when allocating trace memory

It is only when allocating buffers to store normal trace events that care needs to be taken to capture events and blit them to the appropriate place. This is already in place in Writer_NextBuffer(). This also handles cases where traced events are artifacts of Trace allocating memory for other purposes. So the redirect removed by this change was incorrect. Under the correct conditions, it would accelerate events ahead of others due to the SendData() call and result in the incorrect ordering events being sent across the wire.

#rnx

diff --git a/Engine/Source/Runtime/TraceLog/Private/Trace/BlockPool.cpp b/Engine/Source/Runtime/TraceLog/Private/Trace/BlockPool.cpp
--- a/Engine/Source/Runtime/TraceLog/Private/Trace/BlockPool.cpp
+++ b/Engine/Source/Runtime/TraceLog/Private/Trace/BlockPool.cpp
@@ -13,7 +13,7 @@ namespace Trace {
 namespace Private {
 
 ////////////////////////////////////////////////////////////////////////////////
-void*		Writer_MemoryAllocateNoRedirect(SIZE_T, uint32);
+void*		Writer_MemoryAllocate(SIZE_T, uint32);
 void		Writer_MemoryFree(void*, uint32);
 
 
@@ -49,7 +49,7 @@ static uint32							GPoolUsage;			// = 0;
 static FPoolBlockList Writer_AddPageToPool(uint32 PageSize)
 {
 	// The free list is empty so we have to populate it with some new blocks.
-	uint8* PageBase = (uint8*)Writer_MemoryAllocateNoRedirect(PageSize, PLATFORM_CACHE_LINE_SIZE);
+	uint8* PageBase = (uint8*)Writer_MemoryAllocate(PageSize, PLATFORM_CACHE_LINE_SIZE);
 	GPoolUsage += PageSize;
 
 	uint32 BufferSize = GPoolBlockSize;
diff --git a/Engine/Source/Runtime/TraceLog/Private/Trace/Writer.cpp b/Engine/Source/Runtime/TraceLog/Private/Trace/Writer.cpp
--- a/Engine/Source/Runtime/TraceLog/Private/Trace/Writer.cpp
+++ b/Engine/Source/Runtime/TraceLog/Private/Trace/Writer.cpp
@@ -138,7 +138,7 @@ void Writer_MemorySetHooks(decltype(AllocHook) Alloc, decltype(FreeHook) Free)
 }
 
 ////////////////////////////////////////////////////////////////////////////////
-void* Writer_MemoryAllocateNoRedirect(SIZE_T Size, uint32 Alignment)
+void* Writer_MemoryAllocate(SIZE_T Size, uint32 Alignment)
 {
 	void* Ret = nullptr;
 
@@ -184,25 +184,13 @@ void* Writer_MemoryAllocateNoRedirect(SIZE_T Size, uint32 Alignment)
 	return Ret;
 }
 
-////////////////////////////////////////////////////////////////////////////////
-void* Writer_MemoryAllocate(SIZE_T Size, uint32 Alignment)
-{
-	TWriteBufferRedirect<1 << 10> TraceData;
-
-	void* Ret = Writer_MemoryAllocateNoRedirect(Size, Alignment);
-
-	if (TraceData.GetSize())
-	{
-		uint32 ThreadId = Writer_GetThreadId();
-		Writer_SendData(ThreadId, TraceData.GetData(), TraceData.GetSize());
-	}
-
-	return Ret;
-}
-
 ////////////////////////////////////////////////////////////////////////////////
 void Writer_MemoryFree(void* Address, uint32 Size)
 {
+	// Here we are redirecting the per-thread buffer used to write events into. This
+	// path is only used when shutting down, and we neither want to allocate new
+	// memory for traced events or use memory that's been freed.
+
 	TWriteBufferRedirect<1 << 10> TraceData;
 
 #if TRACE_PRIVATE_STOMP
