Fixed deadlock registering a thread-local's destructor.

The PS4 standard C++ library uses an atexit() style call to register destructors of thread-local objects. The implementation takes a lock. This can easily deadlock if memory tracing is enabled as memory tracing will trace when allocating unrelated TLS-related memory. This change uses an alternative scheme to get a notification when a thread terminates.

#rnx

diff --git a/Engine/Source/Runtime/TraceLog/Private/Trace/Writer.cpp b/Engine/Source/Runtime/TraceLog/Private/Trace/Writer.cpp
--- a/Engine/Source/Runtime/TraceLog/Private/Trace/Writer.cpp
+++ b/Engine/Source/Runtime/TraceLog/Private/Trace/Writer.cpp
@@ -83,6 +83,20 @@ static uint32					GUpdateCounter;		// = 0;
 
 
 
+////////////////////////////////////////////////////////////////////////////////
+// When a thread terminates we want to recover its trace buffer. To do this we
+// use a thread_local object whose destructor gets called as the thread ends. On
+// some C++ standard library implementations this is implemented using a thread-
+// specific atexit() call and can involve taking a lock. As tracing can start
+// very early or often be used during shared-object loads, there is a risk of
+// a deadlock initialising the context object. The define below can be used to
+// implement alternatives via a 'ThreadOnThreadExit()' symbol.
+#if !defined(UE_TRACE_USE_TLS_CONTEXT_OBJECT)
+#	define UE_TRACE_USE_TLS_CONTEXT_OBJECT 1
+#endif
+
+#if UE_TRACE_USE_TLS_CONTEXT_OBJECT
+
 ////////////////////////////////////////////////////////////////////////////////
 struct FWriteTlsContext
 {
@@ -124,6 +138,29 @@ uint32 Writer_GetThreadId()
 	return GTlsContext.GetThreadId();
 }
 
+#else // UE_TRACE_USE_TLS_CONTEXT_OBJECT
+
+////////////////////////////////////////////////////////////////////////////////
+void ThreadOnThreadExit(void (*)());
+
+////////////////////////////////////////////////////////////////////////////////
+uint32 Writer_GetThreadId()
+{
+	static thread_local uint32 ThreadId;
+	if (ThreadId)
+	{
+		return ThreadId;
+	}
+
+	ThreadOnThreadExit([] () { Writer_EndThreadBuffer(); });
+
+	static uint32 volatile Counter;
+	ThreadId = AtomicAddRelaxed(&Counter, 1u) + ETransportTid::Bias;
+	return ThreadId;
+}
+
+#endif // UE_TRACE_USE_TLS_CONTEXT_OBJECT
+
 
 
 ////////////////////////////////////////////////////////////////////////////////
