Bumped temporary trace buffer size.

#rnx

diff --git a/Engine/Source/Runtime/TraceLog/Private/Trace/TlsBuffer.cpp b/Engine/Source/Runtime/TraceLog/Private/Trace/TlsBuffer.cpp
--- a/Engine/Source/Runtime/TraceLog/Private/Trace/TlsBuffer.cpp
+++ b/Engine/Source/Runtime/TraceLog/Private/Trace/TlsBuffer.cpp
@@ -57,7 +57,7 @@ static FWriteBuffer* Writer_NextBufferInternal(FWriteBuffer* CurrentBuffer)
 	// memory. Such events should be tightly controlled and kept to a minimum.
 	// There is not a lot of buffer space available in TraceData to avoid
 	// exhausting available stack and leave usable space in the new buffer.
-	TWriteBufferRedirect<512> TraceData;
+	TWriteBufferRedirect<2048> TraceData;
 
 #if !UE_BUILD_SHIPPING && !UE_BUILD_TEST
 	if (CurrentBuffer->ThreadId == decltype(TraceData)::ActiveRedirection)
