Can't have check()'s here as it reverses dependencies

#rnx

diff --git a/Engine/Source/Runtime/TraceLog/Public/Trace/Detail/Field.h b/Engine/Source/Runtime/TraceLog/Public/Trace/Detail/Field.h
--- a/Engine/Source/Runtime/TraceLog/Public/Trace/Detail/Field.h
+++ b/Engine/Source/Runtime/TraceLog/Public/Trace/Detail/Field.h
@@ -184,7 +184,6 @@ public:
 	TField(const FLiteralName& Name, uint32 ReferencedUid)
 		: FieldDesc(Name, Tid, Offset, Size, ReferencedUid)
 	{
-		check(ReferencedUid != 0);
 	}
 };
 
