Fixed GInitialised ambiguity that can arise in unity files.

#rnx

diff --git a/Engine/Source/Runtime/TraceLog/Private/Trace/Channel.cpp b/Engine/Source/Runtime/TraceLog/Private/Trace/Channel.cpp
--- a/Engine/Source/Runtime/TraceLog/Private/Trace/Channel.cpp
+++ b/Engine/Source/Runtime/TraceLog/Private/Trace/Channel.cpp
@@ -40,7 +40,7 @@ UE_TRACE_EVENT_END()
 ///////////////////////////////////////////////////////////////////////////////
 static FChannel* volatile	GHeadChannel;			// = nullptr;
 static FChannel* volatile	GNewChannelList;		// = nullptr;
-static bool 				GInitialized;
+static bool 				GChannelsInitialized;
 
 ////////////////////////////////////////////////////////////////////////////////
 static uint32 GetChannelHash(const ANSICHAR* Input, int32 Length)
@@ -158,7 +158,7 @@ void FChannel::Setup(const ANSICHAR* InChannelName, const InitArgs& InArgs)
 
 	// If channel is initialized after the all channels are disabled (post static init)
 	// this channel needs to be disabled.
-	if (GInitialized)
+	if (GChannelsInitialized)
 	{
 		Enabled = -1;
 	}
@@ -180,7 +180,7 @@ void FChannel::Initialize()
 	// All channels are initialized as enabled (zero), and act like so during
 	// from process start until this method is called (i.e. when Trace is initalized).
 	ToggleAll(false);
-	GInitialized = true;
+	GChannelsInitialized = true;
 }
 
 ///////////////////////////////////////////////////////////////////////////////
