Fixed some iniclude paths so IWYU doesn't get confused.

diff --git a/Engine/Source/Runtime/Core/Public/Math/GenericOctree.h b/Engine/Source/Runtime/Core/Public/Math/GenericOctree.h
--- a/Engine/Source/Runtime/Core/Public/Math/GenericOctree.h
+++ b/Engine/Source/Runtime/Core/Public/Math/GenericOctree.h
@@ -1816,6 +1816,6 @@ public:
 	}
 };
 
-#include "GenericOctree.inl" // IWYU pragma: export
+#include "Math/GenericOctree.inl" // IWYU pragma: export
 
 #include <stddef.h> // IWYU pragma: export
diff --git a/Engine/Source/Runtime/TraceLog/Public/Trace/Trace.inl b/Engine/Source/Runtime/TraceLog/Public/Trace/Trace.inl
--- a/Engine/Source/Runtime/TraceLog/Public/Trace/Trace.inl
+++ b/Engine/Source/Runtime/TraceLog/Public/Trace/Trace.inl
@@ -4,9 +4,9 @@
 
 // IWYU pragma: begin_exports
 #include "CoreTypes.h"
-#include "Detail/Channel.inl"
-#include "Detail/EventNode.h"
-#include "Detail/Important/ImportantLogScope.inl"
-#include "Detail/LogScope.inl"
-#include "Detail/Writer.inl"
+#include "Trace/Detail/Channel.inl"
+#include "Trace/Detail/EventNode.h"
+#include "Trace/Detail/Important/ImportantLogScope.inl"
+#include "Trace/Detail/LogScope.inl"
+#include "Trace/Detail/Writer.inl"
 // IWYU pragma: end_exports
