Fixed out-of-order trace events

After collecting important events, Trace will free the buffers used. If memory tracing is enabled then the Writer_Free() call will immediately send some trace events. If a connection already exists, and this is the first trip through the update loop, then there were some prior events traced to set up the thread. These prior events have lower serial numbers but haven't yet been collected and sent out. Thus Free() ends up sending events too early. There is actually no need to redirect here; Trace doesn't shut itself down, but instead leaves it to the kernel to clean up after the process exits. This is so it can memory events can be traced right to the very end of the process, including during static-exit.

#rnx

diff --git a/Engine/Source/Runtime/TraceLog/Private/Trace/Writer.cpp b/Engine/Source/Runtime/TraceLog/Private/Trace/Writer.cpp
--- a/Engine/Source/Runtime/TraceLog/Private/Trace/Writer.cpp
+++ b/Engine/Source/Runtime/TraceLog/Private/Trace/Writer.cpp
@@ -224,12 +224,6 @@ void* Writer_MemoryAllocate(SIZE_T Size, uint32 Alignment)
 ////////////////////////////////////////////////////////////////////////////////
 void Writer_MemoryFree(void* Address, uint32 Size)
 {
-	// Here we are redirecting the per-thread buffer used to write events into. This
-	// path is only used when shutting down, and we neither want to allocate new
-	// memory for traced events or use memory that's been freed.
-
-	TWriteBufferRedirect<1 << 10> TraceData;
-
 #if TRACE_PRIVATE_STOMP
 	if (Address == nullptr)
 	{
@@ -261,12 +255,6 @@ void Writer_MemoryFree(void* Address, uint32 Size)
 #if TRACE_PRIVATE_STATISTICS
 	AtomicAddRelaxed(&GTraceStatistics.MemoryUsed, uint64(-int64(Size)));
 #endif
-
-	if (TraceData.GetSize())
-	{
-		uint32 ThreadId = Writer_GetThreadId();
-		Writer_SendData(ThreadId, TraceData.GetData(), TraceData.GetSize());
-	}
 }
 
 
