Fixed missing trace events.

Trace uses a thread local object and it destructor as a means to collect buffers when threads end. This usually happens lazily, and usually when fetching a new identified for a thread. On some platforms can result in an allocation and if memory tracing was on, its events would end up in the redirect buffer as the correct buffer hadn't yet been set. But because the redirect had already been memcpy'd, events would get lost.

#rnx

diff --git a/Engine/Source/Runtime/TraceLog/Private/Trace/TlsBuffer.cpp b/Engine/Source/Runtime/TraceLog/Private/Trace/TlsBuffer.cpp
--- a/Engine/Source/Runtime/TraceLog/Private/Trace/TlsBuffer.cpp
+++ b/Engine/Source/Runtime/TraceLog/Private/Trace/TlsBuffer.cpp
@@ -85,6 +85,8 @@ static FWriteBuffer* Writer_NextBufferInternal(FWriteBuffer* CurrentBuffer)
 
 	TraceData.Abandon();
 
+	GTlsWriteBuffer = NextBuffer;
+
 	NextBuffer->Committed = NextBuffer->Cursor;
 
 	if (CurrentBuffer == &GNullWriteBuffer)
@@ -92,8 +94,6 @@ static FWriteBuffer* Writer_NextBufferInternal(FWriteBuffer* CurrentBuffer)
 		NextBuffer->ThreadId = uint16(Writer_GetThreadId());
 		NextBuffer->PrevTimestamp = TimeGetTimestamp();
 
-		GTlsWriteBuffer = NextBuffer;
-
 		UE_TRACE_LOG($Trace, ThreadTiming, TraceLogChannel)
 			<< ThreadTiming.BaseTimestamp(NextBuffer->PrevTimestamp - GStartCycle);
 
@@ -113,8 +113,6 @@ static FWriteBuffer* Writer_NextBufferInternal(FWriteBuffer* CurrentBuffer)
 		NextBuffer->ThreadId = CurrentBuffer->ThreadId;
 		NextBuffer->PrevTimestamp = CurrentBuffer->PrevTimestamp;
 
-		GTlsWriteBuffer = NextBuffer;
-
 		// Retire current buffer.
 		int32 EtxOffset = int32(PTRINT((uint8*)(CurrentBuffer) - CurrentBuffer->Cursor));
 		AtomicStoreRelease(&(CurrentBuffer->EtxOffset), EtxOffset);
