Make default snapshot traces ~20 seconds in editor.

diff --git a/Engine/Source/Runtime/TraceLog/Public/Trace/Trace.h b/Engine/Source/Runtime/TraceLog/Public/Trace/Trace.h
--- a/Engine/Source/Runtime/TraceLog/Public/Trace/Trace.h
+++ b/Engine/Source/Runtime/TraceLog/Public/Trace/Trace.h
@@ -81,7 +81,11 @@ using OnConnectFunc = void(void);
 
 struct FInitializeDesc
 {
-	uint32			TailSizeBytes		= 4 << 20;  // Set to 0 to disable the tail buffer
+#if WITH_EDITOR
+	uint32			TailSizeBytes = 32 << 20;
+#else
+	uint32			TailSizeBytes = 4 << 20; // can be set to 0 to disable the tail buffer
+#endif
 	uint32			ThreadSleepTimeInMS = 0;
 	bool			bUseWorkerThread	= true;
 	bool			bUseImportantCache	= true;
