POSIX standalone fixes.

#rnx

diff --git a/Engine/Source/Runtime/TraceLog/create_standalone.py b/Engine/Source/Runtime/TraceLog/create_standalone.py
--- a/Engine/Source/Runtime/TraceLog/create_standalone.py
+++ b/Engine/Source/Runtime/TraceLog/create_standalone.py
@@ -228,13 +228,19 @@ def _thin(state):
 		(stubs_dir / dep).open("wt").close()
 
 	# Symlink source
-	def symlink(fr, to):
+	def symlink_posix(fr, to):
+		fr = state.dest_dir / "src" / fr
+		fr.symlink_to(to.resolve())
+
+	def symlink_nt(fr, to):
 		fr = state.dest_dir / "src" / fr
 		subprocess.run(
 			("cmd.exe", "/c", "mklink", "/j", fr, to.resolve()),
 			stderr=subprocess.DEVNULL
 		)
 
+	symlink = symlink_nt if os.name == "nt" else symlink_posix
+
 	(state.dest_dir / "src").mkdir(parents=True, exist_ok=True)
 	symlink("trace", state.src_dir)
 	if state.analysis_dir:
diff --git a/Engine/Source/Runtime/TraceLog/standalone_prologue.h b/Engine/Source/Runtime/TraceLog/standalone_prologue.h
--- a/Engine/Source/Runtime/TraceLog/standalone_prologue.h
+++ b/Engine/Source/Runtime/TraceLog/standalone_prologue.h
@@ -56,6 +56,8 @@
 #include <algorithm>
 #include <cstddef>
 #include <cstdint>
+#include <cstring>
+#include <string>
 #include <string_view>
 #include <utility>
 #include <vector>
@@ -310,14 +312,14 @@ inline void*	Realloc(void* ptr, size_t size)			{ return ::realloc(ptr, size); }
 inline void		Free(void* ptr)							{ return ::free(ptr); }
 inline void		Memcpy(void* d, const void* s, size_t n){ ::memcpy(d, s, n); }
 
-template <typename T, typename COMPARE=std::less<T::value_type>>
+template <typename T, typename COMPARE=std::less<typename T::value_type>>
 inline void StableSort(T& container, COMPARE&& compare=COMPARE())
 {
 	std::stable_sort(container.GetData(), container.GetData() + container.Num(),
 		std::forward<COMPARE>(compare));
 }
 
-template <typename T, typename COMPARE=std::less<T::value_type>>
+template <typename T, typename COMPARE=std::less<typename T::value_type>>
 inline void Sort(T& container, COMPARE&& compare=COMPARE())
 {
 	std::sort(container.GetData(), container.GetData() + container.Num(),
@@ -327,7 +329,7 @@ inline void Sort(T& container, COMPARE&& compare=COMPARE())
 template <typename T, typename P>
 inline void SortBy(T& container, P const& p)
 {
-	Sort(container, [&p] (T::value_type& lhs, T::value_type& rhs) {
+	Sort(container, [&p] (typename T::value_type& lhs, typename T::value_type& rhs) {
 		return p(lhs) < p(rhs);
 	});
 }
@@ -335,7 +337,7 @@ inline void SortBy(T& container, P const& p)
 template <typename T, typename V, typename P>
 inline int LowerBoundBy(T const& container, V const& v, P const& p)
 {
-	auto adapter = [&p] (T::value_type const& lhs, V const& v) {
+	auto adapter = [&p] (typename T::value_type const& lhs, V const& v) {
 		return p(lhs) < v;
 	};
 	auto* first = container.GetData();
