Missing include

#rnx

diff --git a/Engine/Source/Runtime/TraceLog/Public/Trace/Detail/Field.h b/Engine/Source/Runtime/TraceLog/Public/Trace/Detail/Field.h
--- a/Engine/Source/Runtime/TraceLog/Public/Trace/Detail/Field.h
+++ b/Engine/Source/Runtime/TraceLog/Public/Trace/Detail/Field.h
@@ -9,6 +9,7 @@
 #include "Atomic.h"
 #include "Protocol.h"
 #include "Templates/UnrealTemplate.h"
+#include "Trace/Trace.h"
 #include "Writer.inl"
 
 /* Statically sized fields (e.g. UE_TRACE_EVENT_FIELD(float[4], Colours)) are
