Renamed a function to something a little more accurate

#rnx

diff --git a/Engine/Source/Runtime/TraceLog/Private/Trace/BlockPool.cpp b/Engine/Source/Runtime/TraceLog/Private/Trace/BlockPool.cpp
--- a/Engine/Source/Runtime/TraceLog/Private/Trace/BlockPool.cpp
+++ b/Engine/Source/Runtime/TraceLog/Private/Trace/BlockPool.cpp
@@ -46,7 +46,7 @@ static uint32							GPoolUsage;			// = 0;
 #undef T_ALIGN
 
 ////////////////////////////////////////////////////////////////////////////////
-static FPoolBlockList Writer_AddPageToPool(uint32 PageSize)
+static FPoolBlockList Writer_AllocateBlockList(uint32 PageSize)
 {
 	// The free list is empty so we have to populate it with some new blocks.
 	uint8* PageBase = (uint8*)Writer_MemoryAllocate(PageSize, PLATFORM_CACHE_LINE_SIZE);
@@ -118,7 +118,7 @@ FWriteBuffer* Writer_AllocateBlockFromPool()
 			continue;
 		}
 
-		FPoolBlockList BlockList = Writer_AddPageToPool(GPoolPageSize);
+		FPoolBlockList BlockList = Writer_AllocateBlockList(GPoolPageSize);
 		Ret = BlockList.Head;
 
 		// And insert the block list into the freelist. 'Block' is now the last block
