Placate TSAN

#rnx
#ushell-unshelve of 22562441

diff --git a/Engine/Source/Runtime/TraceLog/Private/Trace/Writer.cpp b/Engine/Source/Runtime/TraceLog/Private/Trace/Writer.cpp
--- a/Engine/Source/Runtime/TraceLog/Private/Trace/Writer.cpp
+++ b/Engine/Source/Runtime/TraceLog/Private/Trace/Writer.cpp
@@ -640,7 +640,7 @@ static void Writer_WorkerThread()
 {
 	ThreadRegister(TEXT("Trace"), 0, INT_MAX);
 
-	while (!GWorkerThreadQuit)
+	while (!AtomicLoadRelaxed(&GWorkerThreadQuit))
 	{
 		Writer_WorkerUpdate();
 
@@ -667,7 +667,8 @@ static void Writer_WorkerJoin()
 		return;
 	}
 
-	GWorkerThreadQuit = true;
+	AtomicStoreRelaxed(&GWorkerThreadQuit, true);
+
 	ThreadJoin(GWorkerThread);
 	ThreadDestroy(GWorkerThread);
 
