Pacify TSAN

TSAN doesn't realise that the function-scoped static takes care of the lazy-initialisation's thread safety. But this is not a hot path, nor is synchronisation of GInit critical, so we can use relaxed load/stores and TSAN is a happy bunny.

#rnx
#ushell-unshelve of 22562441

diff --git a/Engine/Source/Runtime/TraceLog/Private/Trace/Writer.cpp b/Engine/Source/Runtime/TraceLog/Private/Trace/Writer.cpp
--- a/Engine/Source/Runtime/TraceLog/Private/Trace/Writer.cpp
+++ b/Engine/Source/Runtime/TraceLog/Private/Trace/Writer.cpp
@@ -74,7 +74,7 @@ UE_TRACE_EVENT_END()
 
 
 ////////////////////////////////////////////////////////////////////////////////
-static bool						GInitialized;		// = false;
+static volatile bool			GInitialized;		// = false;
 FStatistics						GTraceStatistics;	// = {};
 uint64							GStartCycle;		// = 0;
 TRACELOG_API uint32 volatile	GLogSerial;			// = 0;
@@ -110,7 +110,7 @@ private:
 ////////////////////////////////////////////////////////////////////////////////
 FWriteTlsContext::~FWriteTlsContext()
 {
-	if (GInitialized)
+	if (AtomicLoadRelaxed(&GInitialized))
 	{
 		Writer_EndThreadBuffer();
 	}
@@ -682,7 +682,7 @@ static void Writer_WorkerJoin()
 ////////////////////////////////////////////////////////////////////////////////
 static void Writer_InternalInitializeImpl()
 {
-	if (GInitialized)
+	if (AtomicLoadRelaxed(&GInitialized))
 	{
 		return;
 	}
@@ -693,7 +693,7 @@ static void Writer_InternalInitializeImpl()
 	Writer_InitializePool();
 	Writer_InitializeControl();
 
-	GInitialized = true;
+	AtomicStoreRelaxed(&GInitialized, true);
 
 	UE_TRACE_LOG($Trace, NewTrace, TraceLogChannel)
 		<< NewTrace.StartCycle(GStartCycle)
@@ -705,7 +705,7 @@ static void Writer_InternalInitializeImpl()
 ////////////////////////////////////////////////////////////////////////////////
 static void Writer_InternalShutdown()
 {
-	if (!GInitialized)
+	if (!AtomicLoadRelaxed(&GInitialized))
 	{
 		return;
 	}
@@ -734,7 +734,7 @@ static void Writer_InternalShutdown()
 	}
 #endif
 
-	GInitialized = false;
+	AtomicStoreRelaxed(&GInitialized, false);
 }
 
 ////////////////////////////////////////////////////////////////////////////////
@@ -742,25 +742,27 @@ void Writer_InternalInitialize()
 {
 	using namespace Private;
 
-	if (!GInitialized)
+	if (AtomicLoadRelaxed(&GInitialized))
 	{
-		static struct FInitializer
-		{
-			FInitializer()
-			{
-				Writer_InternalInitializeImpl();
-			}
-			~FInitializer()
-			{
-				/* We'll not shut anything down here so we can hopefully capture
-				 * any subsequent events. However, we will shutdown the worker
-				 * thread and leave it for something else to call update() (mem
-				 * tracing at time of writing). Windows will have already done
-				 * this implicitly in ExitProcess() anyway. */
-				Writer_WorkerJoin();
-			}
-		} Initializer;
+		return;
 	}
+
+	static struct FInitializer
+	{
+		FInitializer()
+		{
+			Writer_InternalInitializeImpl();
+		}
+		~FInitializer()
+		{
+			/* We'll not shut anything down here so we can hopefully capture
+			 * any subsequent events. However, we will shutdown the worker
+			 * thread and leave it for something else to call update() (mem
+			 * tracing at time of writing). Windows will have already done
+			 * this implicitly in ExitProcess() anyway. */
+			Writer_WorkerJoin();
+		}
+	} Initializer;
 }
 
 ////////////////////////////////////////////////////////////////////////////////
