Don't send update buffers
#rnx

diff --git a/Engine/Source/Runtime/TraceLog/standalone_epilogue.h b/Engine/Source/Runtime/TraceLog/standalone_epilogue.h
--- a/Engine/Source/Runtime/TraceLog/standalone_epilogue.h
+++ b/Engine/Source/Runtime/TraceLog/standalone_epilogue.h
@@ -301,6 +301,9 @@ void ThreadBuffer::Flush(bool Force)
 {
 	using namespace Private;
 
+	if (Cursor == Buffer)
+		return;
+
 	if (!Force && (Cursor <= (Buffer + BufferSize - Overflow)))
 		return;
 
