Correctly communicate update-in-progress change to the worker thread

#rnx

diff --git a/Engine/Source/Runtime/TraceLog/Private/Trace/Writer.cpp b/Engine/Source/Runtime/TraceLog/Private/Trace/Writer.cpp
--- a/Engine/Source/Runtime/TraceLog/Private/Trace/Writer.cpp
+++ b/Engine/Source/Runtime/TraceLog/Private/Trace/Writer.cpp
@@ -601,7 +601,7 @@ void Writer_CallbackOnConnect()
 static UPTRINT			GWorkerThread;		// = 0;
 static volatile bool	GWorkerThreadQuit;	// = false;
 static uint32			GSleepTimeInMS = 17;
-static volatile unsigned int	GUpdateInProgress = 1;	// Don't allow updates until initialized
+static volatile uint32	GUpdateInProgress = 1;	// Don't allow updates until initialized
 
 ////////////////////////////////////////////////////////////////////////////////
 static void Writer_WorkerUpdateInternal()
@@ -789,7 +789,7 @@ void Writer_Initialize(const FInitializeDesc& Desc)
 	GOnConnection = Desc.OnConnectionFunc;
 
 	// Allow the worker thread to start updating 
-	GUpdateInProgress = 0;
+	AtomicStoreRelease(&GUpdateInProgress, uint32(0));
 }
 
 ////////////////////////////////////////////////////////////////////////////////
