Fixed TSAN warnings regarding NextBuffer reads and writes.

NextBuffer is used both to maintain a list of available buffers and to traverse threads' active buffers to collect and send traced data. It was synchronised by other means (and in some rare instances allowed to race harmlessly). TSAN couldn't follow this. Not all warnings were false positives however, there was a risk that loads/stores could be reordered when making a buffer visible to Trace's worker thread.

#rnx

diff --git a/Engine/Source/Runtime/TraceLog/Private/Trace/BlockPool.cpp b/Engine/Source/Runtime/TraceLog/Private/Trace/BlockPool.cpp
--- a/Engine/Source/Runtime/TraceLog/Private/Trace/BlockPool.cpp
+++ b/Engine/Source/Runtime/TraceLog/Private/Trace/BlockPool.cpp
@@ -68,7 +68,7 @@ static FPoolBlockList Writer_AllocateBlockList(uint32 PageSize)
 			break;
 		}
 
-		Buffer->NextBuffer = (FWriteBuffer*)(Block + GPoolBlockSize);
+		AtomicStoreRelaxed(&(Buffer->NextBuffer), (FWriteBuffer*)(Block + GPoolBlockSize));
 		Block += GPoolBlockSize;
 	}
 
@@ -94,7 +94,9 @@ FWriteBuffer* Writer_AllocateBlockFromPool()
 		FWriteBuffer* Owned = AtomicLoadRelaxed(&GPoolFreeList);
 		if (Owned != nullptr)
 		{
-			if (!AtomicCompareExchangeRelaxed(&GPoolFreeList, Owned->NextBuffer, Owned))
+			FWriteBuffer* OwnedNext = AtomicLoadRelaxed(&(Owned->NextBuffer));
+
+			if (!AtomicCompareExchangeAcquire(&GPoolFreeList, OwnedNext, Owned))
 			{
 				PlatformYield();
 				continue;
@@ -124,8 +126,10 @@ FWriteBuffer* Writer_AllocateBlockFromPool()
 		// And insert the block list into the freelist. 'Block' is now the last block
 		for (auto* ListNode = BlockList.Tail;; PlatformYield())
 		{
-			ListNode->NextBuffer = AtomicLoadRelaxed(&GPoolFreeList);
-			if (AtomicCompareExchangeRelease(&GPoolFreeList, Ret->NextBuffer, ListNode->NextBuffer))
+			FWriteBuffer* FreeListValue = AtomicLoadRelaxed(&GPoolFreeList);
+			AtomicStoreRelaxed(&(ListNode->NextBuffer), FreeListValue);
+
+			if (AtomicCompareExchangeRelease(&GPoolFreeList, Ret->NextBuffer, FreeListValue))
 			{
 				break;
 			}
@@ -151,8 +155,10 @@ void Writer_FreeBlockListToPool(FWriteBuffer* Head, FWriteBuffer* Tail)
 {
 	for (FWriteBuffer* ListNode = Tail;; PlatformYield())
 	{
-		ListNode->NextBuffer = AtomicLoadRelaxed(&GPoolFreeList);
-		if (AtomicCompareExchangeRelease(&GPoolFreeList, Head, ListNode->NextBuffer))
+		FWriteBuffer* FreeListValue = AtomicLoadRelaxed(&GPoolFreeList);
+		AtomicStoreRelaxed(&(ListNode->NextBuffer), FreeListValue);
+
+		if (AtomicCompareExchangeRelease(&GPoolFreeList, Head, FreeListValue))
 		{
 			break;
 		}
diff --git a/Engine/Source/Runtime/TraceLog/Private/Trace/TlsBuffer.cpp b/Engine/Source/Runtime/TraceLog/Private/Trace/TlsBuffer.cpp
--- a/Engine/Source/Runtime/TraceLog/Private/Trace/TlsBuffer.cpp
+++ b/Engine/Source/Runtime/TraceLog/Private/Trace/TlsBuffer.cpp
@@ -74,7 +74,7 @@ static FWriteBuffer* Writer_NextBufferInternal(FWriteBuffer* CurrentBuffer)
 	NextBuffer->Cursor = (uint8*)NextBuffer - NextBuffer->Size;
 	NextBuffer->Reaped = NextBuffer->Cursor;
 	NextBuffer->EtxOffset = 0 - int32(sizeof(FWriteBuffer));
-	NextBuffer->NextBuffer = nullptr;
+	AtomicStoreRelaxed(&(NextBuffer->NextBuffer), (FWriteBuffer*)nullptr);
 
 	// Add any capture events into the buffer.
 	if (uint32 RedirectSize = TraceData.GetSize())
@@ -109,9 +109,9 @@ static FWriteBuffer* Writer_NextBufferInternal(FWriteBuffer* CurrentBuffer)
 	}
 	else
 	{
-		CurrentBuffer->NextBuffer = NextBuffer;
 		NextBuffer->ThreadId = CurrentBuffer->ThreadId;
 		NextBuffer->PrevTimestamp = CurrentBuffer->PrevTimestamp;
+		AtomicStoreRelease(&(CurrentBuffer->NextBuffer), NextBuffer);
 
 		// Retire current buffer.
 		int32 EtxOffset = int32(PTRINT((uint8*)(CurrentBuffer) - CurrentBuffer->Cursor));
@@ -159,7 +159,7 @@ void Writer_DrainBuffers()
 
 		void Insert(FWriteBuffer* __restrict Buffer)
 		{
-			Buffer->NextBuffer = Head;
+			AtomicStoreRelaxed(&(Buffer->NextBuffer), Head);
 			Head = Buffer;
 			Tail = (Tail != nullptr) ? Tail : Head;
 		}
@@ -206,7 +206,7 @@ void Writer_DrainBuffers()
 				}
 
 				// Retire the buffer
-				NextBuffer = Buffer->NextBuffer;
+				NextBuffer = AtomicLoadAcquire(&(Buffer->NextBuffer));
 				RetireList.Insert(Buffer);
 			}
 
@@ -242,7 +242,7 @@ void Writer_DrainLocalBuffers()
 
 		void Insert(FWriteBuffer* __restrict Buffer)
 		{
-			Buffer->NextBuffer = Head;
+			AtomicStoreRelaxed(&(Buffer->NextBuffer), Head);
 			Head = Buffer;
 			Tail = (Tail != nullptr) ? Tail : Head;
 		}
@@ -274,7 +274,7 @@ void Writer_DrainLocalBuffers()
 					}
 
 					// Retire the buffer
-					NextBuffer = Buffer->NextBuffer;
+					NextBuffer = AtomicLoadAcquire(&(Buffer->NextBuffer));
 					RetireList.Insert(Buffer);
 				}
 			}
diff --git a/Engine/Source/Runtime/TraceLog/Public/Trace/Detail/Writer.inl b/Engine/Source/Runtime/TraceLog/Public/Trace/Detail/Writer.inl
--- a/Engine/Source/Runtime/TraceLog/Public/Trace/Detail/Writer.inl
+++ b/Engine/Source/Runtime/TraceLog/Public/Trace/Detail/Writer.inl
@@ -14,7 +14,7 @@ struct FWriteBuffer
 	uint8						Overflow[8];
 	uint64						PrevTimestamp;
 	FWriteBuffer* __restrict	NextThread;
-	FWriteBuffer* __restrict	NextBuffer;
+	FWriteBuffer* volatile		NextBuffer;
 	uint8* __restrict			Cursor;
 	uint8* __restrict volatile	Committed;
 	uint8* __restrict			Reaped;
