Pass Send/WriteTo() results to the worker thread in a more conformant manner

Doing so keeps TSAN happy. This isn't really thread safe as there's no synchronisation. However, the only scenario that this would be a problem is when someone is calling some combination of send/write/stop in very quick succession. And if that's going on, something else is very wrong.

#rnx

diff --git a/Engine/Source/Runtime/TraceLog/Private/Trace/Writer.cpp b/Engine/Source/Runtime/TraceLog/Private/Trace/Writer.cpp
--- a/Engine/Source/Runtime/TraceLog/Private/Trace/Writer.cpp
+++ b/Engine/Source/Runtime/TraceLog/Private/Trace/Writer.cpp
@@ -261,7 +261,7 @@ void Writer_MemoryFree(void* Address, uint32 Size)
 
 ////////////////////////////////////////////////////////////////////////////////
 static UPTRINT					GDataHandle;		// = 0
-UPTRINT							GPendingDataHandle;	// = 0
+static volatile UPTRINT			GPendingDataHandle;	// = 0
 
 ////////////////////////////////////////////////////////////////////////////////
 #if TRACE_PRIVATE_BUFFER_SEND
@@ -466,37 +466,41 @@ static void Writer_Close()
 ////////////////////////////////////////////////////////////////////////////////
 static bool Writer_UpdateConnection()
 {
-	if (!GPendingDataHandle)
+	UPTRINT PendingDataHandle = AtomicLoadRelaxed(&GPendingDataHandle);
+
+	if (!PendingDataHandle)
 	{
 		return false;
 	}
 
 	// Is this a close request? So that we capture some of the events around
 	// the closure we will add some inertia before enacting the close.
-	static const uint32 CloseInertia = 2;
-	if (GPendingDataHandle >= (~0ull - CloseInertia))
+	static int32 CloseInertia = 0;
+	if (PendingDataHandle == ~UPTRINT(0))
 	{
-		--GPendingDataHandle;
+		if (CloseInertia <= 0)
+			CloseInertia = 2;
 
-		if (GPendingDataHandle == (~0ull - CloseInertia))
+		--CloseInertia;
+		if (CloseInertia <= 0)
 		{
 			Writer_Close();
-			GPendingDataHandle = 0;
+			AtomicStoreRelaxed(&GPendingDataHandle, UPTRINT(0));
 		}
 
 		return true;
 	}
 
+	AtomicStoreRelaxed(&GPendingDataHandle, UPTRINT(0));
+
 	// Reject the pending connection if we've already got a connection
 	if (GDataHandle)
 	{
-		IoClose(GPendingDataHandle);
-		GPendingDataHandle = 0;
+		IoClose(PendingDataHandle);
 		return false;
 	}
 
-	GDataHandle = GPendingDataHandle;
-	GPendingDataHandle = 0;
+	GDataHandle = PendingDataHandle;
 	if (!Writer_SessionPrologue())
 	{
 		return false;
@@ -810,7 +814,7 @@ void Writer_Update()
 ////////////////////////////////////////////////////////////////////////////////
 bool Writer_SendTo(const ANSICHAR* Host, uint32 Port)
 {
-	if (GPendingDataHandle || GDataHandle)
+	if (AtomicLoadRelaxed(&GPendingDataHandle))
 	{
 		return false;
 	}
@@ -824,14 +828,14 @@ bool Writer_SendTo(const ANSICHAR* Host, uint32 Port)
 		return false;
 	}
 
-	GPendingDataHandle = DataHandle;
+	AtomicStoreRelaxed(&GPendingDataHandle, DataHandle);
 	return true;
 }
 
 ////////////////////////////////////////////////////////////////////////////////
 bool Writer_WriteTo(const ANSICHAR* Path)
 {
-	if (GPendingDataHandle || GDataHandle)
+	if (AtomicLoadRelaxed(&GPendingDataHandle))
 	{
 		return false;
 	}
@@ -844,7 +848,7 @@ bool Writer_WriteTo(const ANSICHAR* Path)
 		return false;
 	}
 
-	GPendingDataHandle = DataHandle;
+	AtomicStoreRelaxed(&GPendingDataHandle, DataHandle);
 	return true;
 }
 
@@ -1001,7 +1005,7 @@ bool Writer_Stop()
 		return false;
 	}
 
-	GPendingDataHandle = ~UPTRINT(0);
+	AtomicStoreRelaxed(&GPendingDataHandle, ~UPTRINT(0));
 	return true;
 }
 
