Address warnings from UBSAN.

Most of the warnings were about unaligned memory access. Such access is not a problem on platforms we care about (an exception being 32 bit ARM and int64's, but we memcpy those already). As such we can use a memcpy() to remove the undefined behaviour and still have the compiler output decent instructions (in >= Development variants). This unfortunately semantically unlinks the structs, but comments have been left should anyone go splunking.

New-Event event writes are already on two-byte boundaries or less, so aligning them up to two bytes in size maintains alignment.

Lastly, UBSAN warns about how TraceLog plays hard and fast with the rules to cast arbitrary pointers to certain types in order to implement the UE_TRACE_LOG() API. We can use C++17's std::launder to communicate intent to the compiler.

#rnx

diff --git a/Engine/Source/Runtime/TraceLog/Private/Trace/EventNode.cpp b/Engine/Source/Runtime/TraceLog/Private/Trace/EventNode.cpp
--- a/Engine/Source/Runtime/TraceLog/Private/Trace/EventNode.cpp
+++ b/Engine/Source/Runtime/TraceLog/Private/Trace/EventNode.cpp
@@ -139,6 +139,7 @@ void FEventNode::Describe() const
 	uint32 EventSize = sizeof(FNewEventEvent);
 	EventSize += sizeof(FNewEventEvent::Fields[0]) * Info->FieldCount;
 	EventSize += NamesSize;
+	EventSize = (EventSize + 1) & ~1; // align to 2 to keep UBSAN happy
 
 	FLogScope LogScope = FLogScope::EnterImpl<FEventInfo::Flag_NoSync>(0, EventSize + sizeof(uint16));
 	auto* Ptr = (uint16*)(LogScope.GetPointer());
diff --git a/Engine/Source/Runtime/TraceLog/Private/Trace/Field.cpp b/Engine/Source/Runtime/TraceLog/Private/Trace/Field.cpp
--- a/Engine/Source/Runtime/TraceLog/Private/Trace/Field.cpp
+++ b/Engine/Source/Runtime/TraceLog/Private/Trace/Field.cpp
@@ -63,13 +63,12 @@ static void Field_WriteAuxData(uint32 Index, int32 Size, CallbackType&& Callback
 		int32 SegmentSize = (Remaining < Size) ? Remaining : Size;
 
 		// Write header
-		auto* Header = (FAuxHeader*)(Buffer->Cursor);
-		Header->Pack = SegmentSize << FAuxHeader::SizeShift;
-		Header->Pack |= Index << FAuxHeader::FieldShift;
-		Header->Uid = uint8(EKnownEventUids::AuxData) << EKnownEventUids::_UidShift;
+		uint32 Pack = SegmentSize << FAuxHeader::SizeShift;
+		Pack |= Index << FAuxHeader::FieldShift;
+		memcpy(Buffer->Cursor, &Pack, sizeof(uint32)); /* FAuxHeader::Pack */
+		Buffer->Cursor[0] = uint8(EKnownEventUids::AuxData) << EKnownEventUids::_UidShift; /* FAuxHeader::Uid */
 		Buffer->Cursor += sizeof(FAuxHeader);
 
-
 		// Write payload data
 		Callback(Buffer->Cursor, SegmentSize);
 		Buffer->Cursor += SegmentSize;
diff --git a/Engine/Source/Runtime/TraceLog/Public/Trace/Detail/Important/ImportantLogScope.inl b/Engine/Source/Runtime/TraceLog/Public/Trace/Detail/Important/ImportantLogScope.inl
--- a/Engine/Source/Runtime/TraceLog/Public/Trace/Detail/Important/ImportantLogScope.inl
+++ b/Engine/Source/Runtime/TraceLog/Public/Trace/Detail/Important/ImportantLogScope.inl
@@ -75,12 +75,11 @@ inline FImportantLogScope FImportantLogScope::EnterImpl(uint32 Uid, uint32 Size)
 	uint8* Out = (uint8*)Buffer - Bias;
 
 	// Event header
-	auto* Header = (FImportantEventHeader*)Out;
-	Header->Uid = uint16(Uid);
-	Header->Size = uint16(Size);
+	uint16 Values16[] = { uint16(Uid), uint16(Size) };
+	memcpy(Out, Values16, sizeof(Values16)); /* FImportantEventHeader::Uid,Size */
 
 	FImportantLogScope Ret;
-	Ret.Ptr = Header->Data;
+	Ret.Ptr = Out + sizeof(FImportantEventHeader);
 	Ret.BufferOffset = int32(PTRINT(Buffer) - PTRINT(Ret.Ptr));
 	Ret.AuxCursor = 0;
 	return Ret;
@@ -114,12 +113,14 @@ struct FImportantLogScope::FFieldSet<FieldMeta, Type[]>
 	{
 		uint32 Size = Num * sizeof(Type);
 
-		auto* Header = (FAuxHeader*)(Scope->Ptr + Scope->AuxCursor);
-		Header->Pack = Size << FAuxHeader::SizeShift;
-		Header->Pack |= (FieldMeta::Index & int32(EIndexPack::NumFieldsMask)) << FAuxHeader::FieldShift;
-		Header->Uid = uint8(EKnownEventUids::AuxData);
+		uint32 Pack = Size << FAuxHeader::SizeShift;
+		Pack |= (FieldMeta::Index & int32(EIndexPack::NumFieldsMask)) << FAuxHeader::FieldShift;
 
-		memcpy(Header + 1, Data, Size);
+		uint8* Out = Scope->Ptr + Scope->AuxCursor;
+		memcpy(Out, &Pack, sizeof(Pack)); /* FAuxHeader::Pack */
+		Out[0] = uint8(EKnownEventUids::AuxData); /* FAuxHeader::Uid */
+
+		memcpy(Out + sizeof(FAuxHeader), Data, Size);
 
 		Scope->AuxCursor += sizeof(FAuxHeader) + Size;
 		Scope->Ptr[Scope->AuxCursor] = uint8(EKnownEventUids::AuxDataTerminal);
@@ -137,12 +138,14 @@ struct FImportantLogScope::FFieldSet<FieldMeta, AnsiString>
 			Length = int32(strlen(String));
 		}
 
-		auto* Header = (FAuxHeader*)(Scope->Ptr + Scope->AuxCursor);
-		Header->Pack = Length << FAuxHeader::SizeShift;
-		Header->Pack |= (FieldMeta::Index & int32(EIndexPack::NumFieldsMask)) << FAuxHeader::FieldShift;
-		Header->Uid = uint8(EKnownEventUids::AuxData);
+		uint32 Pack = Length << FAuxHeader::SizeShift;
+		Pack |= (FieldMeta::Index & int32(EIndexPack::NumFieldsMask)) << FAuxHeader::FieldShift;
+
+		uint8* Out = Scope->Ptr + Scope->AuxCursor;
+		memcpy(Out, &Pack, sizeof(Pack)); /* FAuxHeader::FieldIndex_Size */
+		Out[0] = uint8(EKnownEventUids::AuxData); /* FAuxHeader::Uid */
 
-		memcpy(Header + 1, String, Length);
+		memcpy(Out + sizeof(FAuxHeader), String, Length);
 
 		Scope->AuxCursor += sizeof(FAuxHeader) + Length;
 		Scope->Ptr[Scope->AuxCursor] = uint8(EKnownEventUids::AuxDataTerminal);
@@ -156,12 +159,14 @@ struct FImportantLogScope::FFieldSet<FieldMeta, AnsiString>
 			for (const WIDECHAR* c = String; *c; ++c, ++Length);
 		}
 
-		auto* Header = (FAuxHeader*)(Scope->Ptr + Scope->AuxCursor);
-		Header->Pack = Length << FAuxHeader::SizeShift;
-		Header->Pack |= (FieldMeta::Index & int32(EIndexPack::NumFieldsMask)) << FAuxHeader::FieldShift;
-		Header->Uid = uint8(EKnownEventUids::AuxData);
+		uint32 Pack = Length << FAuxHeader::SizeShift;
+		Pack |= (FieldMeta::Index & int32(EIndexPack::NumFieldsMask)) << FAuxHeader::FieldShift;
 
-		auto* Out = (int8*)(Header + 1);
+		uint8* Out = Scope->Ptr + Scope->AuxCursor;
+		memcpy(Out, &Pack, sizeof(Pack)); /* FAuxHeader::FieldIndex_Size */
+		Out[0] = uint8(EKnownEventUids::AuxData); /* FAuxHeader::Uid */
+
+		Out += sizeof(FAuxHeader);
 		for (int32 i = 0; i < Length; ++i)
 		{
 			*Out = int8(*String);
@@ -188,12 +193,14 @@ struct FImportantLogScope::FFieldSet<FieldMeta, WideString>
 
 		uint32 Size = Length * sizeof(WIDECHAR);
 
-		auto* Header = (FAuxHeader*)(Scope->Ptr + Scope->AuxCursor);
-		Header->Pack = Size << FAuxHeader::SizeShift;
-		Header->Pack |= (FieldMeta::Index & int32(EIndexPack::NumFieldsMask)) << FAuxHeader::FieldShift;
-		Header->Uid = uint8(EKnownEventUids::AuxData);
+		uint32 Pack = Size << FAuxHeader::SizeShift;
+		Pack |= (FieldMeta::Index & int32(EIndexPack::NumFieldsMask)) << FAuxHeader::FieldShift;
+
+		uint8* Out = Scope->Ptr + Scope->AuxCursor;
+		memcpy(Out, &Pack, sizeof(Pack));
+		Out[0] = uint8(EKnownEventUids::AuxData);
 
-		memcpy(Header + 1, String, Size);
+		memcpy(Out + sizeof(FAuxHeader), String, Size);
 
 		Scope->AuxCursor += sizeof(FAuxHeader) + Size;
 		Scope->Ptr[Scope->AuxCursor] = uint8(EKnownEventUids::AuxDataTerminal);
diff --git a/Engine/Source/Runtime/TraceLog/Public/Trace/Detail/LogScope.inl b/Engine/Source/Runtime/TraceLog/Public/Trace/Detail/LogScope.inl
--- a/Engine/Source/Runtime/TraceLog/Public/Trace/Detail/LogScope.inl
+++ b/Engine/Source/Runtime/TraceLog/Public/Trace/Detail/LogScope.inl
@@ -84,10 +84,12 @@ inline void FLogScope::Enter(uint32 Uid, uint32 Size)
 {
 	EnterPrelude<FEventHeaderSync>(Size);
 
-	// Event header
-	auto* Header = (uint16*)(Ptr - sizeof(FEventHeaderSync::SerialHigh));
-	*(uint32*)(Header - 1) = uint32(AtomicAddRelaxed(&GLogSerial, 1u));
-	Header[-2] = uint16(Uid)|int32(EKnownEventUids::Flag_TwoByteUid);
+	uint16 Uid16 = uint16(Uid) | int32(EKnownEventUids::Flag_TwoByteUid);
+	uint32 Serial = uint32(AtomicAddRelaxed(&GLogSerial, 1u));
+
+	// Event header FEventHeaderSync
+	memcpy(Ptr - 3, &Serial, sizeof(Serial)); /* FEventHeaderSync::SerialHigh,SerialLow */
+	memcpy(Ptr - 5, &Uid16,  sizeof(Uid16));  /* FEventHeaderSync::Uid */
 }
 
 ////////////////////////////////////////////////////////////////////////////////
@@ -95,9 +97,11 @@ inline void FLogScope::EnterNoSync(uint32 Uid, uint32 Size)
 {
 	EnterPrelude<FEventHeader>(Size);
 
-	// Event header
+	uint16 Uid16 = uint16(Uid) | int32(EKnownEventUids::Flag_TwoByteUid);
+
+	// Event header FEventHeader
 	auto* Header = (uint16*)(Ptr);
-	Header[-1] = uint16(Uid)|int32(EKnownEventUids::Flag_TwoByteUid);
+	memcpy(Header - 1, &Uid16, sizeof(Uid16)); /* FEventHeader::Uid */
 }
 
 
diff --git a/Engine/Source/Runtime/TraceLog/Public/Trace/Detail/Trace.h b/Engine/Source/Runtime/TraceLog/Public/Trace/Detail/Trace.h
--- a/Engine/Source/Runtime/TraceLog/Public/Trace/Detail/Trace.h
+++ b/Engine/Source/Runtime/TraceLog/Public/Trace/Detail/Trace.h
@@ -6,6 +6,8 @@
 
 #if UE_TRACE_ENABLED
 
+#include "Misc/Launder.h"
+
 #include <type_traits>
 
 namespace UE {
@@ -149,7 +151,7 @@ class FChannel;
 #define TRACE_PRIVATE_LOG_PRELUDE(EnterFunc, LoggerName, EventName, ChannelsExpr, ...) \
 	if (TRACE_PRIVATE_CHANNELEXPR_IS_ENABLED(ChannelsExpr)) \
 		if (auto LogScope = F##LoggerName##EventName##Fields::LogScopeType::EnterFunc<F##LoggerName##EventName##Fields>(__VA_ARGS__)) \
-			if (const auto& __restrict EventName = *(F##LoggerName##EventName##Fields*)(&LogScope)) \
+			if (const auto& __restrict EventName = *UE_LAUNDER((F##LoggerName##EventName##Fields*)(&LogScope))) \
 				((void)EventName),
 
 #define TRACE_PRIVATE_LOG_EPILOG() \
