Updated standalong scripts.

#rnx

diff --git a/Engine/Source/Runtime/TraceLog/standalone_prologue.h b/Engine/Source/Runtime/TraceLog/standalone_prologue.h
--- a/Engine/Source/Runtime/TraceLog/standalone_prologue.h
+++ b/Engine/Source/Runtime/TraceLog/standalone_prologue.h
@@ -145,6 +145,7 @@ using WIDECHAR	= char16_t;
 #define TRACE_ENABLED					1
 #define TRACE_PRIVATE_CONTROL_ENABLED	0
 #define TRACE_PRIVATE_EXTERNAL_LZ4		1
+#define UE_LAUNDER(x)					(x)
 #define UE_TRACE_ENABLED				TRACE_ENABLED
 #define PREPROCESSOR_JOIN(a,b)			PREPROCESSOR_JOIN_IMPL(a,b)
 #define PREPROCESSOR_JOIN_IMPL(a,b)		a##b
