Transparent handling of .cpp files submitted with UTF-8 BOMs.
#rnx

diff --git a/Engine/Source/Runtime/TraceLog/create_standalone.py b/Engine/Source/Runtime/TraceLog/create_standalone.py
--- a/Engine/Source/Runtime/TraceLog/create_standalone.py
+++ b/Engine/Source/Runtime/TraceLog/create_standalone.py
@@ -59,7 +59,7 @@ def _finalize_source(source_files, include_dirs):
 	_spam("Loading code")
 	for source_file in source_files:
 		file = source_file.path
-		for line in file.open("rt"):
+		for line in file.open("rt", encoding="utf-8-sig"):
 			include = _parse_include(line)
 			if not include:
 				if not _exclude_line(line):
