[Insights] Trace: Changed the scoped events with time to use absolute timestamp values. Increased trace protocol to 7.

diff --git a/Engine/Source/Runtime/TraceLog/Public/Trace/Config.h b/Engine/Source/Runtime/TraceLog/Public/Trace/Config.h
--- a/Engine/Source/Runtime/TraceLog/Public/Trace/Config.h
+++ b/Engine/Source/Runtime/TraceLog/Public/Trace/Config.h
@@ -17,5 +17,5 @@
 #endif
 
 #if UE_TRACE_ENABLED
-#	define TRACE_PRIVATE_PROTOCOL_6
+#	define TRACE_PRIVATE_PROTOCOL_7
 #endif
diff --git a/Engine/Source/Runtime/TraceLog/Public/Trace/Detail/LogScope.inl b/Engine/Source/Runtime/TraceLog/Public/Trace/Detail/LogScope.inl
--- a/Engine/Source/Runtime/TraceLog/Public/Trace/Detail/LogScope.inl
+++ b/Engine/Source/Runtime/TraceLog/Public/Trace/Detail/LogScope.inl
@@ -165,17 +165,18 @@ inline FScopedStampedLogScope::~FScopedStampedLogScope()
 		return;
 	}
 
-	FWriteBuffer* Buffer = Writer_GetBuffer();
-
-	uint64 Stamp = Writer_GetTimestamp(Buffer);
+	uint8 LeaveUid = uint8(EKnownEventUids::LeaveScope_TA << EKnownEventUids::_UidShift);
+	uint64 Stamp = TimeGetTimestamp();
 
-	if (UNLIKELY(int32((uint8*)Buffer - Buffer->Cursor) < int32(sizeof(Stamp))))
+	FWriteBuffer* Buffer = Writer_GetBuffer();
+	if (UNLIKELY(int32((uint8*)Buffer - Buffer->Cursor) < int32(sizeof(LeaveUid)) + int32(sizeof(Stamp))))
 	{
 		Buffer = Writer_NextBuffer();
 	}
 
-	Stamp <<= 8;
-	Stamp += uint8(EKnownEventUids::LeaveScope_T) << EKnownEventUids::_UidShift;
+	Buffer->Cursor[0] = LeaveUid;
+	Buffer->Cursor += sizeof(LeaveUid);
+
 	memcpy((uint64*)(Buffer->Cursor), &Stamp, sizeof(Stamp));
 	Buffer->Cursor += sizeof(Stamp);
 
@@ -223,17 +224,19 @@ FORCENOINLINE auto FLogScope::ScopedEnter()
 template <class EventType>
 FORCENOINLINE auto FLogScope::ScopedStampedEnter()
 {
+	uint8 EnterUid = uint8(EKnownEventUids::EnterScope_TA << EKnownEventUids::_UidShift);
 	uint64 Stamp;
 
 	FWriteBuffer* Buffer = Writer_GetBuffer();
-	if (UNLIKELY(int32((uint8*)Buffer - Buffer->Cursor) < int32(sizeof(Stamp))))
+	if (UNLIKELY(int32((uint8*)Buffer - Buffer->Cursor) < int32(sizeof(EnterUid)) + int32(sizeof(Stamp))))
 	{
 		Buffer = Writer_NextBuffer();
 	}
 
-	Stamp = Writer_GetTimestamp(Buffer);
-	Stamp <<= 8;
-	Stamp += uint8(EKnownEventUids::EnterScope_T) << EKnownEventUids::_UidShift;
+	Buffer->Cursor[0] = EnterUid;
+	Buffer->Cursor += sizeof(EnterUid);
+
+	Stamp = TimeGetTimestamp();
 	memcpy((uint64*)(Buffer->Cursor), &Stamp, sizeof(Stamp));
 	Buffer->Cursor += sizeof(Stamp);
 
diff --git a/Engine/Source/Runtime/TraceLog/Public/Trace/Detail/Protocol.h b/Engine/Source/Runtime/TraceLog/Public/Trace/Detail/Protocol.h
--- a/Engine/Source/Runtime/TraceLog/Public/Trace/Detail/Protocol.h
+++ b/Engine/Source/Runtime/TraceLog/Public/Trace/Detail/Protocol.h
@@ -17,6 +17,7 @@
 #include "Protocols/Protocol4.h"
 #include "Protocols/Protocol5.h"
 #include "Protocols/Protocol6.h"
+#include "Protocols/Protocol7.h"
 
 #if defined(_MSC_VER)
 	#pragma warning(pop)
diff --git a/Engine/Source/Runtime/TraceLog/Public/Trace/Detail/Protocols/Protocol6.h b/Engine/Source/Runtime/TraceLog/Public/Trace/Detail/Protocols/Protocol6.h
--- a/Engine/Source/Runtime/TraceLog/Public/Trace/Detail/Protocols/Protocol6.h
+++ b/Engine/Source/Runtime/TraceLog/Public/Trace/Detail/Protocols/Protocol6.h
@@ -1,11 +1,9 @@
 ﻿// Copyright Epic Games, Inc. All Rights Reserved.
+
 #pragma once
 
 // HEADER_UNIT_SKIP - Not included directly
 
-#include "Protocol4.h"
-#include "Protocol5.h"
-
 namespace UE {
 namespace Trace {
 
diff --git a/Engine/Source/Runtime/TraceLog/Public/Trace/Detail/Protocols/Protocol7.h b/Engine/Source/Runtime/TraceLog/Public/Trace/Detail/Protocols/Protocol7.h
new file mode 100644
--- /dev/null
+++ b/Engine/Source/Runtime/TraceLog/Public/Trace/Detail/Protocols/Protocol7.h
@@ -0,0 +1,64 @@
+// Copyright Epic Games, Inc. All Rights Reserved.
+
+#pragma once
+
+// HEADER_UNIT_SKIP - Not included directly
+
+namespace UE {
+namespace Trace {
+
+#if defined(TRACE_PRIVATE_PROTOCOL_7)
+inline
+#endif
+namespace Protocol7
+{
+
+////////////////////////////////////////////////////////////////////////////////
+enum EProtocol : uint8 { Id = 7 };
+
+////////////////////////////////////////////////////////////////////////////////
+using Protocol6::EFieldType;
+using Protocol6::FEventHeader;
+using Protocol6::FImportantEventHeader;
+using Protocol6::FEventHeaderSync;
+using Protocol6::FAuxHeader;
+using Protocol6::EEventFlags;
+using Protocol6::EFieldFamily;
+using Protocol6::FNewEventEvent;
+
+////////////////////////////////////////////////////////////////////////////////
+struct EKnownEventUids
+{
+	static const uint16 Flag_TwoByteUid = 1 << 0;
+	static const uint16 _UidShift = 1;
+	enum : uint16
+	{
+		NewEvent = 0,			// same as Protocol5
+		AuxData,				// same as Protocol5
+		_Unused0,
+		AuxDataTerminal,		// same as Protocol5
+		EnterScope,				// same as Protocol5
+		LeaveScope,				// same as Protocol5
+		EnterScope_TA,			// new in Protocol7
+		LeaveScope_TA,			// new in Protocol7
+		_Unused1,
+		_Unused2,
+		_Unused3,
+		_Unused4,
+		_Unused5,
+		_Unused6,
+		_Unused7,
+		_Unused8,
+		_WellKnownNum,			// same as Protocol5
+	};
+	static const uint16 User = _WellKnownNum;
+	static const uint16 Max = (1 << (16 - _UidShift)) - 1;
+	static const uint16 Invalid = Max;
+	static_assert(User == Protocol5::EKnownEventUids::User, "Protocol7::EKnownEventUids should extend Protocol5");
+};
+
+////////////////////////////////////////////////////////////////////////////////
+
+} // namespace Protocol7
+} // namespace Trace
+} // namespace UE
