Moved modifications of settings based on context to a more appropriate location, in keeping with architectural intent.
#rnx

diff --git a/Engine/Source/Runtime/Core/Private/ProfilingDebugging/TraceAuxiliary.cpp b/Engine/Source/Runtime/Core/Private/ProfilingDebugging/TraceAuxiliary.cpp
--- a/Engine/Source/Runtime/Core/Private/ProfilingDebugging/TraceAuxiliary.cpp
+++ b/Engine/Source/Runtime/Core/Private/ProfilingDebugging/TraceAuxiliary.cpp
@@ -1529,6 +1529,9 @@ void FTraceAuxiliary::Initialize(const TCHAR* CommandLine)
 
 	// Initialize Trace
 	UE::Trace::FInitializeDesc Desc;
+#if WITH_EDITOR
+	Desc.TailSizeBytes = 32 << 20;
+#endif
 	SetupInitFromConfig(Desc);
 	
 	Desc.bUseWorkerThread = bShouldStartWorkerThread;
diff --git a/Engine/Source/Runtime/TraceLog/Public/Trace/Trace.h b/Engine/Source/Runtime/TraceLog/Public/Trace/Trace.h
--- a/Engine/Source/Runtime/TraceLog/Public/Trace/Trace.h
+++ b/Engine/Source/Runtime/TraceLog/Public/Trace/Trace.h
@@ -81,11 +81,7 @@ using OnConnectFunc = void(void);
 
 struct FInitializeDesc
 {
-#if WITH_EDITOR
-	uint32			TailSizeBytes = 32 << 20;
-#else
-	uint32			TailSizeBytes = 4 << 20; // can be set to 0 to disable the tail buffer
-#endif
+	uint32			TailSizeBytes		= 4 << 20; // can be set to 0 to disable the tail buffer
 	uint32			ThreadSleepTimeInMS = 0;
 	bool			bUseWorkerThread	= true;
 	bool			bUseImportantCache	= true;
