Fixed standalone macros leaking out to doctest.
#rnx

diff --git a/Engine/Source/Runtime/TraceLog/standalone_epilogue.h b/Engine/Source/Runtime/TraceLog/standalone_epilogue.h
--- a/Engine/Source/Runtime/TraceLog/standalone_epilogue.h
+++ b/Engine/Source/Runtime/TraceLog/standalone_epilogue.h
@@ -31,6 +31,8 @@ inline void SerializeToCborImpl(TArray<uint8>&, const IAnalyzer::FEventData&, ui
 #	endif
 #endif
 
+#undef check
+
 #endif // TRACE_UE_COMPAT_LAYER
 
 
