Increase Trace write buffer to 2048 on Switch. Mitigates a thread deadlock issue by altering timing.

diff --git a/Engine/Source/Runtime/TraceLog/Private/Trace/Important/SharedBuffer.cpp b/Engine/Source/Runtime/TraceLog/Private/Trace/Important/SharedBuffer.cpp
--- a/Engine/Source/Runtime/TraceLog/Private/Trace/Important/SharedBuffer.cpp
+++ b/Engine/Source/Runtime/TraceLog/Private/Trace/Important/SharedBuffer.cpp
@@ -8,6 +8,10 @@
 #include "Trace/Detail/Atomic.h"
 #include "Trace/Detail/Important/ImportantLogScope.inl"
 
+#ifndef PLATFORM_TRACE_WRITER_BUFFER_SIZE
+#define PLATFORM_TRACE_WRITER_BUFFER_SIZE				( 1*1024 )
+#endif // PLATFORM_TRACE_WRITER_BUFFER_SIZE
+
 namespace UE {
 namespace Trace {
 namespace Private {
@@ -22,7 +26,7 @@ static FSharedBuffer	GNullSharedBuffer	= { 0, FSharedBuffer::RefInit };
 FSharedBuffer* volatile GSharedBuffer		= &GNullSharedBuffer;
 static FSharedBuffer*	GTailBuffer;		// = nullptr
 static uint32			GTailPreSent;		// = 0
-static const uint32		GBlockSize			= 1024;		// Block size must be a power of two!
+static const uint32		GBlockSize			= PLATFORM_TRACE_WRITER_BUFFER_SIZE;		// Block size must be a power of two!
 extern FStatistics		GTraceStatistics;
 
 ////////////////////////////////////////////////////////////////////////////////
