The platform yield used on some platforms was not actually yielding and causing occasional live locks.

#rnx

diff --git a/Engine/Source/Runtime/TraceLog/Public/Trace/Detail/Atomic.h b/Engine/Source/Runtime/TraceLog/Public/Trace/Detail/Atomic.h
--- a/Engine/Source/Runtime/TraceLog/Public/Trace/Detail/Atomic.h
+++ b/Engine/Source/Runtime/TraceLog/Public/Trace/Detail/Atomic.h
@@ -9,6 +9,10 @@
 #	include <immintrin.h>
 #endif
 
+#if !defined(TRACE_PRIVATE_THREAD_YIELD)
+#	define TRACE_PRIVATE_THREAD_YIELD 0
+#endif
+
 namespace UE {
 namespace Trace {
 namespace Private {
@@ -31,7 +35,10 @@ void							PlatformYield();
 ////////////////////////////////////////////////////////////////////////////////
 inline void PlatformYield()
 {
-#if PLATFORM_CPU_X86_FAMILY
+#if TRACE_PRIVATE_THREAD_YIELD
+	extern void ThreadYield();
+	ThreadYield();
+#elif PLATFORM_CPU_X86_FAMILY
 	_mm_pause();
 #elif PLATFORM_CPU_ARM_FAMILY
 #	if defined(_MSC_VER) && !defined(__clang__) // MSVC
