Removed duplicate include statement

#rnx

diff --git a/Engine/Source/Runtime/TraceLog/Private/Trace/Detail/Android/AndroidTrace.cpp b/Engine/Source/Runtime/TraceLog/Private/Trace/Detail/Android/AndroidTrace.cpp
--- a/Engine/Source/Runtime/TraceLog/Private/Trace/Detail/Android/AndroidTrace.cpp
+++ b/Engine/Source/Runtime/TraceLog/Private/Trace/Detail/Android/AndroidTrace.cpp
@@ -12,7 +12,6 @@
 #include <sys/socket.h>
 #include <sys/stat.h>
 #include <unistd.h>
-#include <errno.h>
 
 namespace UE {
 namespace Trace {
