Add session and trace metadata fields to traces.

diff --git a/Engine/Source/Runtime/TraceLog/Private/Trace/Trace.cpp b/Engine/Source/Runtime/TraceLog/Private/Trace/Trace.cpp
--- a/Engine/Source/Runtime/TraceLog/Private/Trace/Trace.cpp
+++ b/Engine/Source/Runtime/TraceLog/Private/Trace/Trace.cpp
@@ -25,6 +25,7 @@ bool	Writer_WriteTo(const ANSICHAR*, uint32);
 bool	Writer_WriteSnapshotTo(const ANSICHAR*);
 bool	Writer_SendSnapshotTo(const ANSICHAR*, uint32);	
 bool	Writer_IsTracing();
+bool	Writer_IsTracingTo(uint32 (&OutSessionGuid)[4], uint32 (&OutTraceGuid)[4]);
 bool	Writer_Stop();
 uint32	Writer_GetThreadId();
 
@@ -120,6 +121,12 @@ bool IsTracing()
 	return Private::Writer_IsTracing();
 }
 
+////////////////////////////////////////////////////////////////////////////////
+bool IsTracingTo(uint32 (&OutSessionGuid)[4], uint32 (&OutTraceGuid)[4])
+{
+	return Private::Writer_IsTracingTo(OutSessionGuid, OutTraceGuid);
+}
+	
 ////////////////////////////////////////////////////////////////////////////////
 bool Stop()
 {
diff --git a/Engine/Source/Runtime/TraceLog/Private/Trace/Writer.cpp b/Engine/Source/Runtime/TraceLog/Private/Trace/Writer.cpp
--- a/Engine/Source/Runtime/TraceLog/Private/Trace/Writer.cpp
+++ b/Engine/Source/Runtime/TraceLog/Private/Trace/Writer.cpp
@@ -62,6 +62,11 @@ static bool		Writer_SessionPrologue();
 void			Writer_FreeBlockListToPool(FWriteBuffer*, FWriteBuffer*);
 
 
+////////////////////////////////////////////////////////////////////////////////
+struct FTraceGuid
+{
+	uint32 Bits[4];	
+};
 
 ////////////////////////////////////////////////////////////////////////////////
 UE_TRACE_EVENT_BEGIN($Trace, NewTrace, Important|NoSync)
@@ -162,6 +167,25 @@ uint32 Writer_GetThreadId()
 #endif // UE_TRACE_USE_TLS_CONTEXT_OBJECT
 
 
+////////////////////////////////////////////////////////////////////////////////
+void Writer_CreateGuid(FTraceGuid* OutGuid)
+{
+	// This is not thread safe. Should only be accessed from the writer thread.
+	static uint64 State = TimeGetTimestamp();
+	// L'Ecuyer, Pierre (1999). "Tables of Linear Congruential Generators of Different Sizes and Good Lattice Structure"
+	// corrected with errata
+	// Assuming m = 2e64
+	constexpr uint64 C = 0x369DEA0F31A53F85;	
+	constexpr uint64 I = 1ull;
+	const uint64 Ts = TimeGetTimestamp();
+
+	const uint64 NewState = (State * C + I);
+	State = NewState;
+
+	*(uint64*)&OutGuid->Bits[0] = NewState;
+	OutGuid->Bits[2] = (0x4 << 28) | uint32(Ts >> 32);
+	OutGuid->Bits[3] = uint32(Ts);
+}
 
 ////////////////////////////////////////////////////////////////////////////////
 void*			(*AllocHook)(SIZE_T, uint32);			// = nullptr
@@ -427,6 +451,8 @@ static void Writer_DescribeAnnounce()
 static int8			GSyncPacketCountdown;	// = 0
 static const int8	GNumSyncPackets			= 3;
 static OnConnectFunc*	GOnConnection = nullptr;
+static FTraceGuid GSessionGuid; // = {0, 0, 0, 0};
+static FTraceGuid GTraceGuid; // = {0, 0, 0, 0};
 
 ////////////////////////////////////////////////////////////////////////////////
 static void Writer_SendSync()
@@ -504,6 +530,9 @@ static bool Writer_UpdateConnection()
 		return false;
 	}
 
+	// Generate Guid for new connection
+	Writer_CreateGuid(&GTraceGuid);
+
 	GDataHandle = PendingDataHandle;
 	if (!Writer_SessionPrologue())
 	{
@@ -557,16 +586,25 @@ static bool Writer_SessionPrologue()
 	struct FHandshake
 	{
 		uint32 Magic			= '2' | ('C' << 8) | ('R' << 16) | ('T' << 24);
-		uint16 MetadataSize		= uint16(4); //  = sizeof(MetadataField0 + ControlPort)
+		uint16 MetadataSize		= uint16(MetadataSizeSum);
 		uint16 MetadataField0	= uint16(sizeof(ControlPort) | (ControlPortFieldId << 8));
 		uint16 ControlPort		= uint16(Writer_GetControlPort());
+		uint16 MetadataField1	= uint16(sizeof(FTraceGuid) | (SessionGuidFieldId << 8));
+		uint32 SessionGuid[4];
+		uint16 MetadataField2	= uint16(sizeof(FTraceGuid) | (TraceGuidFieldId << 8));
+		uint32 TraceGuid[4];
 		enum
 		{
-			Size				= 10,
+			MetadataSizeSum		= 2 + 2 + 2 + 16 + 2 + 16,
+			Size				= MetadataSizeSum + 4 + 2,
 			ControlPortFieldId	= 0,
+			SessionGuidFieldId	= 1,
+			TraceGuidFieldId	= 2,
 		};
 	};
 	FHandshake Handshake;
+	memcpy(&Handshake.SessionGuid, &GSessionGuid, sizeof(FTraceGuid));
+	memcpy(&Handshake.TraceGuid, &GTraceGuid, sizeof(FTraceGuid));
 	bool bOk = IoWrite(GDataHandle, &Handshake, FHandshake::Size);
 
 	// Stream header
@@ -796,6 +834,16 @@ void Writer_Initialize(const FInitializeDesc& Desc)
 		Writer_WorkerCreate();
 	}
 
+	// Store the session guid if specified, otherwise generate one
+	if (!Desc.SessionGuid[0] & !Desc.SessionGuid[1] & !Desc.SessionGuid[2] & !Desc.SessionGuid[3])
+	{
+		Writer_CreateGuid(&GSessionGuid);
+	}
+	else
+	{
+		memcpy(&GSessionGuid, &Desc.SessionGuid, sizeof(FTraceGuid));
+	}
+
 	// Store callback on connection
 	GOnConnection = Desc.OnConnectionFunc;
 
@@ -1088,6 +1136,18 @@ bool Writer_IsTracing()
 {
 	return GDataHandle != 0 || GPendingDataHandle != 0;
 }
+	
+////////////////////////////////////////////////////////////////////////////////
+bool Writer_IsTracingTo(uint32 (&OutSessionGuid)[4], uint32 (&OutTraceGuid)[4])
+{
+	if (Writer_IsTracing())
+	{
+		memcpy(&OutSessionGuid, &GSessionGuid, sizeof(OutSessionGuid));
+		memcpy(&OutTraceGuid, &GTraceGuid, sizeof(OutTraceGuid));
+		return true;
+	}
+	return false;
+}
 
 ////////////////////////////////////////////////////////////////////////////////
 bool Writer_Stop()
diff --git a/Engine/Source/Runtime/TraceLog/Public/Trace/Trace.h b/Engine/Source/Runtime/TraceLog/Public/Trace/Trace.h
--- a/Engine/Source/Runtime/TraceLog/Public/Trace/Trace.h
+++ b/Engine/Source/Runtime/TraceLog/Public/Trace/Trace.h
@@ -85,6 +85,7 @@ struct FInitializeDesc
 	uint32			ThreadSleepTimeInMS = 0;
 	bool			bUseWorkerThread	= true;
 	bool			bUseImportantCache	= true;
+	uint32			SessionGuid[4]		= {0,0,0,0}; // leave as zero to generate random
 	OnConnectFunc*	OnConnectionFunc	= nullptr;
 };
 
@@ -125,6 +126,7 @@ UE_TRACE_API bool	WriteTo(const TCHAR* Path, uint16 Flags=FSendFlags::None) UE_T
 UE_TRACE_API bool	WriteSnapshotTo(const TCHAR* Path) UE_TRACE_IMPL(false);
 UE_TRACE_API bool	SendSnapshotTo(const TCHAR* Host, uint32 Port) UE_TRACE_IMPL(false);	
 UE_TRACE_API bool	IsTracing() UE_TRACE_IMPL(false);
+UE_TRACE_API bool	IsTracingTo(uint32 (&OutSessionGuid)[4], uint32 (&OutTraceGuid)[4]) UE_TRACE_IMPL(false);
 UE_TRACE_API bool	Stop() UE_TRACE_IMPL(false);
 UE_TRACE_API bool	IsChannel(const TCHAR* ChanneName) UE_TRACE_IMPL(false);
 UE_TRACE_API bool	ToggleChannel(const TCHAR* ChannelName, bool bEnabled) UE_TRACE_IMPL(false);
