Remove padding in handshake header

diff --git a/Engine/Source/Runtime/TraceLog/Private/Trace/Writer.cpp b/Engine/Source/Runtime/TraceLog/Private/Trace/Writer.cpp
--- a/Engine/Source/Runtime/TraceLog/Private/Trace/Writer.cpp
+++ b/Engine/Source/Runtime/TraceLog/Private/Trace/Writer.cpp
@@ -590,9 +590,9 @@ static bool Writer_SessionPrologue()
 		uint16 MetadataField0	= uint16(sizeof(ControlPort) | (ControlPortFieldId << 8));
 		uint16 ControlPort		= uint16(Writer_GetControlPort());
 		uint16 MetadataField1	= uint16(sizeof(FTraceGuid) | (SessionGuidFieldId << 8));
-		uint32 SessionGuid[4];
+		uint8 SessionGuid[16];	// Avoid padding
 		uint16 MetadataField2	= uint16(sizeof(FTraceGuid) | (TraceGuidFieldId << 8));
-		uint32 TraceGuid[4];
+		uint8 TraceGuid[16];	// Avoid padding
 		enum
 		{
 			MetadataSizeSum		= 2 + 2 + 2 + 16 + 2 + 16,
