Update Guid generation for TraceLog to use all bits better

#rnx

diff --git a/Engine/Source/Runtime/TraceLog/Private/Trace/Writer.cpp b/Engine/Source/Runtime/TraceLog/Private/Trace/Writer.cpp
--- a/Engine/Source/Runtime/TraceLog/Private/Trace/Writer.cpp
+++ b/Engine/Source/Runtime/TraceLog/Private/Trace/Writer.cpp
@@ -171,20 +171,30 @@ uint32 Writer_GetThreadId()
 void Writer_CreateGuid(FTraceGuid* OutGuid)
 {
 	// This is not thread safe. Should only be accessed from the writer thread.
+	// This initialized the prng with the current timestamp. In theory two machines could initialize on the exact same time
+	// producing the same sequence of guids.
 	static uint64 State = TimeGetTimestamp();
 	// L'Ecuyer, Pierre (1999). "Tables of Linear Congruential Generators of Different Sizes and Good Lattice Structure"
 	// corrected with errata
 	// Assuming m = 2e64
 	constexpr uint64 C = 0x369DEA0F31A53F85;	
 	constexpr uint64 I = 1ull;
-	const uint64 Ts = TimeGetTimestamp();
 
-	const uint64 NewState = (State * C + I);
-	State = NewState;
+	const uint64 TopBits = State * C + I;
+	const uint64 BottomBits = TopBits * C + I;
+	State = BottomBits;
 
-	*(uint64*)&OutGuid->Bits[0] = NewState;
-	OutGuid->Bits[2] = (0x4 << 28) | uint32(Ts >> 32);
-	OutGuid->Bits[3] = uint32(Ts);
+	*(uint64*)&OutGuid->Bits[0] = TopBits;
+	*(uint64*)&OutGuid->Bits[2] = BottomBits;
+
+	constexpr uint8 Version = 0x40; //Version 4, 4 bits
+	constexpr uint8 VersionMask = 0xf0;
+	constexpr uint8 Variant = 0x80; //Variant 1, 2 bits
+	constexpr uint8 VariantMask = 0xc0;
+
+	uint8* Octets = (uint8*)OutGuid;
+	Octets[6] = Version | (~VersionMask & Octets[6]); // Octet 9
+	Octets[8] = Variant | (~VariantMask & Octets[8]); // Octet 7
 }
 
 ////////////////////////////////////////////////////////////////////////////////
