Added FTraceAuxiliary::Panic that disables all channels without logging. Needed on some platforms to make sure no trace code will be executed in their crash handling code

#rnx

diff --git a/Engine/Source/Runtime/Core/Private/ProfilingDebugging/TraceAuxiliary.cpp b/Engine/Source/Runtime/Core/Private/ProfilingDebugging/TraceAuxiliary.cpp
--- a/Engine/Source/Runtime/Core/Private/ProfilingDebugging/TraceAuxiliary.cpp
+++ b/Engine/Source/Runtime/Core/Private/ProfilingDebugging/TraceAuxiliary.cpp
@@ -1735,6 +1735,11 @@ void FTraceAuxiliary::GetActiveChannelsString(FStringBuilderBase& String)
 #endif
 }
 
+void FTraceAuxiliary::Panic()
+{
+	UE::Trace::Panic();
+}
+
 ////////////////////////////////////////////////////////////////////////////////
 void FTraceAuxiliary::TryAutoConnect()
 {
diff --git a/Engine/Source/Runtime/Core/Public/ProfilingDebugging/TraceAuxiliary.h b/Engine/Source/Runtime/Core/Public/ProfilingDebugging/TraceAuxiliary.h
--- a/Engine/Source/Runtime/Core/Public/ProfilingDebugging/TraceAuxiliary.h
+++ b/Engine/Source/Runtime/Core/Public/ProfilingDebugging/TraceAuxiliary.h
@@ -176,6 +176,11 @@ public:
 	 */
 	static void	GetActiveChannelsString(FStringBuilderBase& String);
 
+	/**
+	 * Used when process is panicking. Stops all tracing immediately to avoid further allocations. Process is not
+	 * expected to continue after this call.
+	 */
+	 static void Panic();
 
 	/**
 	 * Delegate that triggers when a connection is established. Gives subscribers a chance to trace events that appear
diff --git a/Engine/Source/Runtime/TraceLog/Private/Trace/Channel.cpp b/Engine/Source/Runtime/TraceLog/Private/Trace/Channel.cpp
--- a/Engine/Source/Runtime/TraceLog/Private/Trace/Channel.cpp
+++ b/Engine/Source/Runtime/TraceLog/Private/Trace/Channel.cpp
@@ -202,6 +202,25 @@ void FChannel::ToggleAll(bool bEnabled)
 	}
 }
 
+///////////////////////////////////////////////////////////////////////////////
+void FChannel::PanicDisableAll()
+{
+	using namespace Private;
+
+	FChannel* ChannelLists[] =
+	{
+		AtomicLoadAcquire(&GNewChannelList),
+		AtomicLoadAcquire(&GHeadChannel),
+	};
+	for (FChannel* Channel : ChannelLists)
+	{
+		for (; Channel != nullptr; Channel = (FChannel*)(Channel->Next))
+		{
+			AtomicStoreRelaxed(&Channel->Enabled, -1);
+		}
+	}
+}
+
 ///////////////////////////////////////////////////////////////////////////////
 FChannel* FChannel::FindChannel(const ANSICHAR* ChannelName)
 {
diff --git a/Engine/Source/Runtime/TraceLog/Private/Trace/Trace.cpp b/Engine/Source/Runtime/TraceLog/Private/Trace/Trace.cpp
--- a/Engine/Source/Runtime/TraceLog/Private/Trace/Trace.cpp
+++ b/Engine/Source/Runtime/TraceLog/Private/Trace/Trace.cpp
@@ -71,6 +71,12 @@ void Shutdown()
 	Private::Writer_Shutdown();
 }
 
+////////////////////////////////////////////////////////////////////////////////
+void Panic()
+{
+	FChannel::PanicDisableAll();
+}
+
 ////////////////////////////////////////////////////////////////////////////////
 void Update()
 {
diff --git a/Engine/Source/Runtime/TraceLog/Public/Trace/Detail/Channel.h b/Engine/Source/Runtime/TraceLog/Public/Trace/Detail/Channel.h
--- a/Engine/Source/Runtime/TraceLog/Public/Trace/Detail/Channel.h
+++ b/Engine/Source/Runtime/TraceLog/Public/Trace/Detail/Channel.h
@@ -60,6 +60,7 @@ public:
 	void				Announce() const;
 	static bool			Toggle(const ANSICHAR* ChannelName, bool bEnabled);
 	static void			ToggleAll(bool bEnabled);
+	static void			PanicDisableAll(); // Disabled channels wont be logged with UE_TRACE_LOG
 	static FChannel*	FindChannel(const ANSICHAR* ChannelName);
 	UE_DEPRECATED(5.2, "Please use the ChannelIterCallback overload for enumerating channels.")
 	static void			EnumerateChannels(ChannelIterFunc Func, void* User);
diff --git a/Engine/Source/Runtime/TraceLog/Public/Trace/Trace.h b/Engine/Source/Runtime/TraceLog/Public/Trace/Trace.h
--- a/Engine/Source/Runtime/TraceLog/Public/Trace/Trace.h
+++ b/Engine/Source/Runtime/TraceLog/Public/Trace/Trace.h
@@ -119,6 +119,7 @@ UE_TRACE_API void	SetMemoryHooks(AllocFunc Alloc, FreeFunc Free) UE_TRACE_IMPL()
 UE_TRACE_API void	Initialize(const FInitializeDesc& Desc) UE_TRACE_IMPL();
 UE_TRACE_API void	StartWorkerThread() UE_TRACE_IMPL();	
 UE_TRACE_API void	Shutdown() UE_TRACE_IMPL();
+UE_TRACE_API void	Panic() UE_TRACE_IMPL();
 UE_TRACE_API void	Update() UE_TRACE_IMPL();
 UE_TRACE_API void	GetStatistics(FStatistics& Out) UE_TRACE_IMPL();
 UE_TRACE_API bool	SendTo(const TCHAR* Host, uint32 Port=0, uint16 Flags=FSendFlags::None) UE_TRACE_IMPL(false);
