Remove WriteTo functionality in control.

Removed ability to write to file from TraceLog control layer. This was not actively used, and proved to be a security concern.

diff --git a/Engine/Source/Runtime/TraceLog/Private/Trace/Control.cpp b/Engine/Source/Runtime/TraceLog/Private/Trace/Control.cpp
--- a/Engine/Source/Runtime/TraceLog/Private/Trace/Control.cpp
+++ b/Engine/Source/Runtime/TraceLog/Private/Trace/Control.cpp
@@ -297,16 +297,6 @@ void Writer_InitializeControl()
 		}
 	);
 
-	Writer_ControlAddCommand("WriteTo", nullptr,
-		[] (void*, uint32 ArgC, ANSICHAR const* const* ArgV)
-		{
-			if (ArgC > 0)
-			{
-				Writer_WriteTo(ArgV[0]);
-			}
-		}
-	);
-
 	Writer_ControlAddCommand("Stop", nullptr,
 		[] (void*, uint32 ArgC, ANSICHAR const* const* ArgV)
 		{
