Add message events in TraceLog

Adds ability for users of the library to listen to "message events". The messages currently provides extended information on error states (broken connections, files, etc).

diff --git a/Engine/Source/Runtime/TraceLog/Private/Trace/Control.cpp b/Engine/Source/Runtime/TraceLog/Private/Trace/Control.cpp
--- a/Engine/Source/Runtime/TraceLog/Private/Trace/Control.cpp
+++ b/Engine/Source/Runtime/TraceLog/Private/Trace/Control.cpp
@@ -5,6 +5,7 @@
 #if UE_TRACE_ENABLED
 
 #include "Trace/Platform.h"
+#include "Trace/Message.h"
 
 #include "Misc/CString.h"
 #include "Templates/UnrealTemplate.h"
@@ -126,6 +127,8 @@ static bool Writer_ControlListen()
 
 	if (!GControlListen)
 	{
+		//This unfortunately triggers on editor shutdown needlessly spamming the log
+		//UE_TRACE_ERRORMESSAGE_F(ListenFail, GetLastErrorCode(), "Port: %d", GControlPort);
 		GControlState = EControlState::Failed;
 		return false;
 	}
@@ -169,6 +172,7 @@ static void Writer_ControlRecv()
 		int32 Recvd = IoRead(GControlSocket, Head, ReadSize);
 		if (Recvd <= 0)
 		{
+			UE_TRACE_ERRORMESSAGE_F(ReadError, GetLastErrorCode(), "While reading remote command. Read returned %d", Recvd);
 			IoClose(GControlSocket);
 			GControlSocket = 0;
 			GControlState = EControlState::Listening;
diff --git a/Engine/Source/Runtime/TraceLog/Private/Trace/Detail/Android/AndroidTrace.cpp b/Engine/Source/Runtime/TraceLog/Private/Trace/Detail/Android/AndroidTrace.cpp
--- a/Engine/Source/Runtime/TraceLog/Private/Trace/Detail/Android/AndroidTrace.cpp
+++ b/Engine/Source/Runtime/TraceLog/Private/Trace/Detail/Android/AndroidTrace.cpp
@@ -222,6 +222,21 @@ UPTRINT FileOpen(const ANSICHAR* Path)
 
 	return UPTRINT(Out + 1);
 }
+	
+////////////////////////////////////////////////////////////////////////////////
+int32 GetLastErrorCode()
+{
+	return errno;
+}
+
+////////////////////////////////////////////////////////////////////////////////
+bool GetErrorMessage(char* OutBuffer, uint32 BufferSize, int32 ErrorCode)
+{
+	const char* ErrorMessage = std::strerror(ErrorCode);
+	const bool bResult = strncpy(OutBuffer, ErrorMessage, BufferSize) != 0;
+	OutBuffer[BufferSize-1] = 0;
+	return bResult;
+}
 
 } // namespace Private
 } // namespace Trace
diff --git a/Engine/Source/Runtime/TraceLog/Private/Trace/Detail/Apple/AppleTrace.cpp b/Engine/Source/Runtime/TraceLog/Private/Trace/Detail/Apple/AppleTrace.cpp
--- a/Engine/Source/Runtime/TraceLog/Private/Trace/Detail/Apple/AppleTrace.cpp
+++ b/Engine/Source/Runtime/TraceLog/Private/Trace/Detail/Apple/AppleTrace.cpp
@@ -255,6 +255,21 @@ UPTRINT FileOpen(const ANSICHAR* Path)
 	return UPTRINT(Out + 1);
 }
 
+////////////////////////////////////////////////////////////////////////////////
+int32 GetLastErrorCode()
+{
+	return errno;
+}
+
+////////////////////////////////////////////////////////////////////////////////
+bool GetErrorMessage(char* OutBuffer, uint32 BufferSize, int32 ErrorCode)
+{
+	const char* ErrorMessage = std::strerror(ErrorCode);
+	const bool bResult = strncpy(OutBuffer, ErrorMessage, BufferSize) != 0;
+	OutBuffer[BufferSize-1] = 0;
+	return bResult;
+}
+
 } // namespace Private
 } // namespace Trace
 } // namespace UE
diff --git a/Engine/Source/Runtime/TraceLog/Private/Trace/Detail/Unix/UnixTrace.cpp b/Engine/Source/Runtime/TraceLog/Private/Trace/Detail/Unix/UnixTrace.cpp
--- a/Engine/Source/Runtime/TraceLog/Private/Trace/Detail/Unix/UnixTrace.cpp
+++ b/Engine/Source/Runtime/TraceLog/Private/Trace/Detail/Unix/UnixTrace.cpp
@@ -4,6 +4,7 @@
 
 #if UE_TRACE_ENABLED && PLATFORM_UNIX
 
+#include <cstring>
 #include <errno.h>
 #include <fcntl.h>
 #include <netdb.h>
@@ -254,6 +255,21 @@ UPTRINT FileOpen(const ANSICHAR* Path)
 	return UPTRINT(Out + 1);
 }
 
+////////////////////////////////////////////////////////////////////////////////
+int32 GetLastErrorCode()
+{
+	return errno;
+}
+
+////////////////////////////////////////////////////////////////////////////////
+bool GetErrorMessage(char* OutBuffer, uint32 BufferSize, int32 ErrorCode)
+{
+	const char* ErrorMessage = std::strerror(ErrorCode);
+	const bool bResult = strncpy(OutBuffer, ErrorMessage, BufferSize) != 0;
+	OutBuffer[BufferSize-1] = 0;
+	return bResult;
+}
+
 } // namespace Private
 } // namespace Trace
 } // namespace UE
diff --git a/Engine/Source/Runtime/TraceLog/Private/Trace/Detail/Windows/WindowsTrace.cpp b/Engine/Source/Runtime/TraceLog/Private/Trace/Detail/Windows/WindowsTrace.cpp
--- a/Engine/Source/Runtime/TraceLog/Private/Trace/Detail/Windows/WindowsTrace.cpp
+++ b/Engine/Source/Runtime/TraceLog/Private/Trace/Detail/Windows/WindowsTrace.cpp
@@ -19,6 +19,19 @@ namespace UE {
 namespace Trace {
 namespace Private {
 
+struct FSetLastErrorScope
+{
+	FSetLastErrorScope(DWORD InError)
+		: Error(InError){}
+	~FSetLastErrorScope()
+	{
+		SetLastError(Error);
+	}
+
+private:
+	DWORD Error;
+};
+
 ////////////////////////////////////////////////////////////////////////////////
 UPTRINT ThreadCreate(const ANSICHAR*, void (*Entry)())
 {
@@ -128,12 +141,14 @@ UPTRINT TcpSocketConnect(const ANSICHAR* Host, uint16 Port)
 	int Result = connect(Socket, Info->ai_addr, int(Info->ai_addrlen));
 	if (Result == SOCKET_ERROR)
 	{
+		FSetLastErrorScope _(GetLastError());
 		closesocket(Socket);
 		return 0;
 	}
 
 	if (!TcpSocketSetNonBlocking(Socket, 0))
 	{
+		FSetLastErrorScope _(GetLastError());
 		closesocket(Socket);
 		return 0;
 	}
@@ -160,6 +175,7 @@ UPTRINT TcpSocketListen(uint16 Port)
 	int Result = bind(Socket, (SOCKADDR*)&SockAddr, sizeof(SockAddr));
 	if (Result == INVALID_SOCKET)
 	{
+		FSetLastErrorScope _(GetLastError());
 		closesocket(Socket);
 		return 0;
 	}
@@ -167,12 +183,14 @@ UPTRINT TcpSocketListen(uint16 Port)
 	Result = listen(Socket, 1);
 	if (Result == INVALID_SOCKET)
 	{
+		FSetLastErrorScope _(GetLastError());
 		closesocket(Socket);
 		return 0;
 	}
 
 	if (!TcpSocketSetNonBlocking(Socket, 1))
 	{
+		FSetLastErrorScope _(GetLastError());
 		closesocket(Socket);
 		return 0;
 	}
@@ -193,6 +211,7 @@ int32 TcpSocketAccept(UPTRINT Socket, UPTRINT& Out)
 
 	if (!TcpSocketSetNonBlocking(Inner, 0))
 	{
+		FSetLastErrorScope _(GetLastError());
 		closesocket(Inner);
 		return 0;
 	}
@@ -265,6 +284,19 @@ UPTRINT FileOpen(const ANSICHAR* Path)
 	return UPTRINT(Out) + 1;
 }
 
+	
+////////////////////////////////////////////////////////////////////////////////
+int32 GetLastErrorCode()
+{
+	return ::GetLastError();
+}
+
+////////////////////////////////////////////////////////////////////////////////
+bool GetErrorMessage(char* OutBuffer, uint32 BufferSize, int32 ErrorCode)
+{
+	return FormatMessageA(FORMAT_MESSAGE_FROM_SYSTEM, NULL, ErrorCode, 0, OutBuffer, BufferSize, NULL) != 0;
+}
+	
 } // namespace Private
 } // namespace Trace
 } // namespace UE
diff --git a/Engine/Source/Runtime/TraceLog/Private/Trace/Message.cpp b/Engine/Source/Runtime/TraceLog/Private/Trace/Message.cpp
new file mode 100644
--- /dev/null
+++ b/Engine/Source/Runtime/TraceLog/Private/Trace/Message.cpp
@@ -0,0 +1,37 @@
+﻿
+// Copyright Epic Games, Inc. All Rights Reserved.
+
+#include "Message.h"
+
+#if UE_TRACE_ENABLED
+
+#include "Trace/Trace.h"
+
+namespace UE {
+namespace Trace {
+namespace Private {
+
+////////////////////////////////////////////////////////////////////////////////
+static OnMessageFunc* GMessageFunc = nullptr;
+	
+////////////////////////////////////////////////////////////////////////////////
+void Message_SetCallback(OnMessageFunc Callback)
+{
+	GMessageFunc = Callback;
+}
+
+////////////////////////////////////////////////////////////////////////////////
+void Message_Send(EMessageType Type, const char* TypeStr, const char* Description /* = nullptr */)
+{
+	if (!GMessageFunc)
+	{
+		return;
+	}
+
+	FMessageEvent Message { Type, TypeStr, Description };
+	GMessageFunc(Message);
+}
+
+} } } // namespace UE::Trace::Private
+
+#endif // UE_TRACE_ENABLED
\ No newline at end of file
diff --git a/Engine/Source/Runtime/TraceLog/Private/Trace/Message.h b/Engine/Source/Runtime/TraceLog/Private/Trace/Message.h
new file mode 100644
--- /dev/null
+++ b/Engine/Source/Runtime/TraceLog/Private/Trace/Message.h
@@ -0,0 +1,69 @@
+﻿// Copyright Epic Games, Inc. All Rights Reserved.
+
+#pragma once
+
+#include "Trace/Config.h"
+#include "Trace/Trace.h"
+
+#if UE_TRACE_ENABLED
+
+namespace UE {
+namespace Trace {
+namespace Private {
+
+bool GetErrorMessage(char* OutBuffer, uint32 BufferSize, int32 ErrorCode);	
+	
+/**
+ * Utility macro to format messages
+ */
+
+#define UE_TRACE_MESSAGE_FMT_MAX_SIZE 512
+#define UE_TRACE_MESSAGE_ERR_MAX_SIZE 256
+
+#define UE_TRACE_MESSAGE(Type, Msg) \
+	Message_Send(EMessageType::Type, #Type, Msg);
+	
+#define UE_TRACE_MESSAGE_F(Type, Fmt, ...) \
+	{\
+		char Buff[UE_TRACE_MESSAGE_FMT_MAX_SIZE];\
+		snprintf(Buff, UE_TRACE_MESSAGE_FMT_MAX_SIZE, Fmt, __VA_ARGS__);\
+		Message_Send(EMessageType::Type, #Type, Buff);\
+	}
+	
+#define UE_TRACE_ERRORMESSAGE(Type, ErrorCode)\
+	{\
+		char ErrorMessageBuffer[UE_TRACE_MESSAGE_ERR_MAX_SIZE] = {'\0'};\
+		GetErrorMessage(ErrorMessageBuffer, UE_TRACE_MESSAGE_ERR_MAX_SIZE, ErrorCode);\
+		UE_TRACE_MESSAGE_F(Type, "(error code %d): '%s'", ErrorCode, ErrorMessageBuffer)\
+	}
+	
+#define UE_TRACE_ERRORMESSAGE_F(Type, ErrorCode, ContextFmt, ...)\
+	{\
+		char ErrorMessageBuffer[UE_TRACE_MESSAGE_ERR_MAX_SIZE];\
+		if (GetErrorMessage(ErrorMessageBuffer, UE_TRACE_MESSAGE_ERR_MAX_SIZE, ErrorCode))\
+		{\
+			char FinalBuff[UE_TRACE_MESSAGE_FMT_MAX_SIZE];\
+			snprintf(FinalBuff, UE_TRACE_MESSAGE_FMT_MAX_SIZE, ContextFmt ": '%s'", __VA_ARGS__, ErrorMessageBuffer);\
+			Message_Send(EMessageType::Type, #Type, FinalBuff);\
+		}\
+		else\
+		{\
+			UE_TRACE_MESSAGE_F(Type, "(error code %d)", ErrorCode)\
+		}\
+	}
+
+/**
+ * Initialize messages. Saves the provided callback if specified.
+ */
+void Message_SetCallback(UE::Trace::OnMessageFunc Callback);
+	
+/**
+ * Sends a message to the registered callback. Description string is not
+ * expected to outlive the function call.
+ */
+void Message_Send(EMessageType Type, const char* TypeStr, const char* Description = nullptr);
+	
+
+} } } // namespace UE::Trace::Private
+
+#endif // UE_TRACE_ENABLED
diff --git a/Engine/Source/Runtime/TraceLog/Private/Trace/Platform.h b/Engine/Source/Runtime/TraceLog/Private/Trace/Platform.h
--- a/Engine/Source/Runtime/TraceLog/Private/Trace/Platform.h
+++ b/Engine/Source/Runtime/TraceLog/Private/Trace/Platform.h
@@ -33,6 +33,10 @@ void	IoClose(UPTRINT Handle);
 
 ////////////////////////////////////////////////////////////////////////////////
 UPTRINT	FileOpen(const ANSICHAR* Path);
+	
+////////////////////////////////////////////////////////////////////////////////
+int32	GetLastErrorCode();
+bool	GetErrorMessage(char* OutBuffer, uint32 BufferSize, int32 ErrorCode);	
 
 } // namespace Private
 } // namespace Trace
diff --git a/Engine/Source/Runtime/TraceLog/Private/Trace/Trace.cpp b/Engine/Source/Runtime/TraceLog/Private/Trace/Trace.cpp
--- a/Engine/Source/Runtime/TraceLog/Private/Trace/Trace.cpp
+++ b/Engine/Source/Runtime/TraceLog/Private/Trace/Trace.cpp
@@ -15,6 +15,7 @@ namespace Private
 {
 
 ////////////////////////////////////////////////////////////////////////////////
+void	Message_SetCallback(OnMessageFunc Callback);
 void	Writer_MemorySetHooks(AllocFunc, FreeFunc);
 void	Writer_Initialize(const FInitializeDesc&);
 void	Writer_WorkerCreate();
@@ -58,6 +59,12 @@ void SetMemoryHooks(AllocFunc Alloc, FreeFunc Free)
 	Private::Writer_MemorySetHooks(Alloc, Free);
 }
 
+////////////////////////////////////////////////////////////////////////////////
+void SetMessageCallback(OnMessageFunc MessageFunc)
+{
+	Private::Message_SetCallback(MessageFunc);
+}
+
 ////////////////////////////////////////////////////////////////////////////////
 void Initialize(const FInitializeDesc& Desc)
 {
diff --git a/Engine/Source/Runtime/TraceLog/Private/Trace/Writer.cpp b/Engine/Source/Runtime/TraceLog/Private/Trace/Writer.cpp
--- a/Engine/Source/Runtime/TraceLog/Private/Trace/Writer.cpp
+++ b/Engine/Source/Runtime/TraceLog/Private/Trace/Writer.cpp
@@ -1,5 +1,6 @@
 // Copyright Epic Games, Inc. All Rights Reserved.
 
+#include "Message.h"
 #include "Trace/Config.h"
 
 #if UE_TRACE_ENABLED
@@ -308,6 +309,7 @@ static bool Writer_FlushSendBuffer()
 	{
 		if (!IoWrite(GDataHandle, GSendBuffer, GSendBufferCursor - GSendBuffer))
 		{
+			UE_TRACE_ERRORMESSAGE(WriteError, GetLastErrorCode());
 			IoClose(GDataHandle);
 			GDataHandle = 0;
 			return false;
@@ -355,6 +357,7 @@ static void Writer_SendDataImpl(const void* Data, uint32 Size)
 #else
 	if (!IoWrite(GDataHandle, Data, Size))
 	{
+		UE_TRACE_ERRORMESSAGE(WriteError, GetLastErrorCode());
 		IoClose(GDataHandle);
 		GDataHandle = 0;
 	}
@@ -626,6 +629,7 @@ static bool Writer_SessionPrologue()
 
 	if (!bOk)
 	{
+		UE_TRACE_ERRORMESSAGE(WriteError, GetLastErrorCode());
 		IoClose(GDataHandle);
 		GDataHandle = 0;
 		return false;
@@ -674,7 +678,7 @@ static void Writer_WorkerUpdateInternal()
 
 #if TRACE_PRIVATE_BUFFER_SEND
 	const uint32 FlushSendBufferCadenceMask = 8-1; // Flush every 8 calls
-	if( (++GUpdateCounter & FlushSendBufferCadenceMask) == 0)
+	if((++GUpdateCounter & FlushSendBufferCadenceMask) == 0 && GDataHandle != 0)
 	{
 		Writer_FlushSendBuffer();
 	}
@@ -909,12 +913,14 @@ bool Writer_SendTo(const ANSICHAR* Host, uint32 Flags, uint32 Port)
 	UPTRINT DataHandle = TcpSocketConnect(Host, uint16(Port));
 	if (!DataHandle)
 	{
+		UE_TRACE_ERRORMESSAGE_F(ConnectError, GetLastErrorCode(), "Connecting to host (%s:%d)", Host, Port);
 		return false;
 	}
 
 	DataHandle = Writer_PackSendFlags(DataHandle, Flags);
 	if (!DataHandle)
 	{
+		UE_TRACE_MESSAGE(ConnectError, "Handle was unexpectedly using MSB flags.");
 		return false;
 	}
 
@@ -935,6 +941,7 @@ bool Writer_WriteTo(const ANSICHAR* Path, uint32 Flags)
 	UPTRINT DataHandle = FileOpen(Path);
 	if (!DataHandle)
 	{
+		UE_TRACE_ERRORMESSAGE_F(FileOpenError, GetLastErrorCode(), "Opening file (%s)", Path);
 		return false;
 	}
 
@@ -1093,6 +1100,7 @@ bool Writer_WriteSnapshot(const FSnapshotTarget& Target)
 		// Write the file header
 		if (!GDataHandle || !Writer_SessionPrologue())
 		{
+			UE_TRACE_ERRORMESSAGE(FileOpenError, GetLastErrorCode());
 			return false;
 		}
 
diff --git a/Engine/Source/Runtime/TraceLog/Public/Trace/Trace.h b/Engine/Source/Runtime/TraceLog/Public/Trace/Trace.h
--- a/Engine/Source/Runtime/TraceLog/Public/Trace/Trace.h
+++ b/Engine/Source/Runtime/TraceLog/Public/Trace/Trace.h
@@ -79,6 +79,29 @@ TEventRef<IdType> MakeEventRef(IdType InId, uint32 InTypeId)
 
 using OnConnectFunc = void(void);
 
+enum class EMessageType : uint8
+{
+	Reserved = 0,
+	WriteError,
+	ReadError,
+	ConnectError,
+	ListenError,
+	EstablishError,
+	FileOpenError,
+};
+	
+struct FMessageEvent
+{
+	/** Type of message */
+	EMessageType		Type;
+	/** Type of message stringified */
+	const char*			TypeStr;
+	/** Clarifying message, may be null for some message types. Pointer only valide during callback. */
+	const char*			Description;
+};
+
+using OnMessageFunc = void(const FMessageEvent&);
+
 struct FInitializeDesc
 {
 	uint32			TailSizeBytes		= 4 << 20; // can be set to 0 to disable the tail buffer
@@ -116,6 +139,7 @@ struct FSendFlags
 };
 
 UE_TRACE_API void	SetMemoryHooks(AllocFunc Alloc, FreeFunc Free) UE_TRACE_IMPL();
+UE_TRACE_API void	SetMessageCallback(OnMessageFunc MessageFunc) UE_TRACE_IMPL();
 UE_TRACE_API void	Initialize(const FInitializeDesc& Desc) UE_TRACE_IMPL();
 UE_TRACE_API void	StartWorkerThread() UE_TRACE_IMPL();	
 UE_TRACE_API void	Shutdown() UE_TRACE_IMPL();
