Print format should be using unsigned.

#rnx

diff --git a/Engine/Source/Runtime/TraceLog/Private/Trace/Writer.cpp b/Engine/Source/Runtime/TraceLog/Private/Trace/Writer.cpp
--- a/Engine/Source/Runtime/TraceLog/Private/Trace/Writer.cpp
+++ b/Engine/Source/Runtime/TraceLog/Private/Trace/Writer.cpp
@@ -913,7 +913,7 @@ bool Writer_SendTo(const ANSICHAR* Host, uint32 Flags, uint32 Port)
 	UPTRINT DataHandle = TcpSocketConnect(Host, uint16(Port));
 	if (!DataHandle)
 	{
-		UE_TRACE_ERRORMESSAGE_F(ConnectError, GetLastErrorCode(), "Connecting to host (%s:%d)", Host, Port);
+		UE_TRACE_ERRORMESSAGE_F(ConnectError, GetLastErrorCode(), "Connecting to host (%s:%u)", Host, Port);
 		return false;
 	}
 
