Fix nomaster issue on Mac

diff --git a/Engine/Source/Runtime/TraceLog/Private/Trace/Detail/Apple/AppleTrace.cpp b/Engine/Source/Runtime/TraceLog/Private/Trace/Detail/Apple/AppleTrace.cpp
--- a/Engine/Source/Runtime/TraceLog/Private/Trace/Detail/Apple/AppleTrace.cpp
+++ b/Engine/Source/Runtime/TraceLog/Private/Trace/Detail/Apple/AppleTrace.cpp
@@ -5,6 +5,7 @@
 #if UE_TRACE_ENABLED && PLATFORM_APPLE
 
 #include <arpa/inet.h>
+#include <string.h>
 #include <errno.h>
 #include <fcntl.h>
 #include <mach/mach.h>
