Nomaster fix for Android

diff --git a/Engine/Source/Runtime/TraceLog/Private/Trace/Detail/Android/AndroidTrace.cpp b/Engine/Source/Runtime/TraceLog/Private/Trace/Detail/Android/AndroidTrace.cpp
--- a/Engine/Source/Runtime/TraceLog/Private/Trace/Detail/Android/AndroidTrace.cpp
+++ b/Engine/Source/Runtime/TraceLog/Private/Trace/Detail/Android/AndroidTrace.cpp
@@ -5,6 +5,7 @@
 #if UE_TRACE_ENABLED && PLATFORM_ANDROID
 
 #include <arpa/inet.h>
+#include <string.h>
 #include <errno.h>
 #include <fcntl.h>
 #include <pthread.h>
