Android and Mac doesn't implement strerror in std namespace

#rnx

diff --git a/Engine/Source/Runtime/TraceLog/Private/Trace/Detail/Android/AndroidTrace.cpp b/Engine/Source/Runtime/TraceLog/Private/Trace/Detail/Android/AndroidTrace.cpp
--- a/Engine/Source/Runtime/TraceLog/Private/Trace/Detail/Android/AndroidTrace.cpp
+++ b/Engine/Source/Runtime/TraceLog/Private/Trace/Detail/Android/AndroidTrace.cpp
@@ -233,7 +233,7 @@ int32 GetLastErrorCode()
 ////////////////////////////////////////////////////////////////////////////////
 bool GetErrorMessage(char* OutBuffer, uint32 BufferSize, int32 ErrorCode)
 {
-	const char* ErrorMessage = std::strerror(ErrorCode);
+	const char* ErrorMessage = strerror(ErrorCode);
 	const bool bResult = strncpy(OutBuffer, ErrorMessage, BufferSize) != 0;
 	OutBuffer[BufferSize-1] = 0;
 	return bResult;
diff --git a/Engine/Source/Runtime/TraceLog/Private/Trace/Detail/Apple/AppleTrace.cpp b/Engine/Source/Runtime/TraceLog/Private/Trace/Detail/Apple/AppleTrace.cpp
--- a/Engine/Source/Runtime/TraceLog/Private/Trace/Detail/Apple/AppleTrace.cpp
+++ b/Engine/Source/Runtime/TraceLog/Private/Trace/Detail/Apple/AppleTrace.cpp
@@ -265,7 +265,7 @@ int32 GetLastErrorCode()
 ////////////////////////////////////////////////////////////////////////////////
 bool GetErrorMessage(char* OutBuffer, uint32 BufferSize, int32 ErrorCode)
 {
-	const char* ErrorMessage = std::strerror(ErrorCode);
+	const char* ErrorMessage = strerror(ErrorCode);
 	const bool bResult = strncpy(OutBuffer, ErrorMessage, BufferSize) != 0;
 	OutBuffer[BufferSize-1] = 0;
 	return bResult;
