Do not bother sending sync message if connection is closed.

There is no check if the connection happens to be closed during Writer_WorkerUpdateInternal.

diff --git a/Engine/Source/Runtime/TraceLog/Private/Trace/Writer.cpp b/Engine/Source/Runtime/TraceLog/Private/Trace/Writer.cpp
--- a/Engine/Source/Runtime/TraceLog/Private/Trace/Writer.cpp
+++ b/Engine/Source/Runtime/TraceLog/Private/Trace/Writer.cpp
@@ -469,7 +469,7 @@ static FTraceGuid GTraceGuid; // = {0, 0, 0, 0};
 ////////////////////////////////////////////////////////////////////////////////
 static void Writer_SendSync()
 {
-	if (GSyncPacketCountdown <= 0)
+	if (GSyncPacketCountdown <= 0 || !GDataHandle)
 	{
 		return;
 	}
