Remove false error message

Removing error message when reading from non-blocking sockets. This will return EAGAIN which is expected.

diff --git a/Engine/Source/Runtime/TraceLog/Private/Trace/Control.cpp b/Engine/Source/Runtime/TraceLog/Private/Trace/Control.cpp
--- a/Engine/Source/Runtime/TraceLog/Private/Trace/Control.cpp
+++ b/Engine/Source/Runtime/TraceLog/Private/Trace/Control.cpp
@@ -173,7 +173,6 @@ static void Writer_ControlRecv()
 		int32 Recvd = IoRead(GControlSocket, Head, ReadSize);
 		if (Recvd <= 0)
 		{
-			UE_TRACE_ERRORMESSAGE_F(ReadError, GetLastErrorCode(), "While reading remote command. Read returned %d", Recvd);
 			IoClose(GControlSocket);
 			GControlSocket = 0;
 			GControlState = EControlState::Listening;
