Trace: Improved trace data writer to drain only the buffers available when we start draining buffers of a thread. This avoids cases where a thread floods the writer with new buffers.

diff --git a/Engine/Source/Runtime/TraceLog/Private/Trace/TlsBuffer.cpp b/Engine/Source/Runtime/TraceLog/Private/Trace/TlsBuffer.cpp
--- a/Engine/Source/Runtime/TraceLog/Private/Trace/TlsBuffer.cpp
+++ b/Engine/Source/Runtime/TraceLog/Private/Trace/TlsBuffer.cpp
@@ -197,9 +197,22 @@ void Writer_DrainBuffers()
 			NextThread = Buffer->NextThread;
 			uint32 ThreadId = Buffer->ThreadId;
 
+			// Count how many buffers are available now, for the current thread.
+			uint32 DrainBufferLimit = 1;
+			for (FWriteBuffer* __restrict CrtBuffer = Buffer; CrtBuffer != nullptr; ++DrainBufferLimit)
+			{
+				CrtBuffer = AtomicLoadAcquire(&(CrtBuffer->NextBuffer));
+			}
+
 			// For each of the thread's buffers...
 			for (FWriteBuffer* __restrict NextBuffer; Buffer != nullptr; Buffer = NextBuffer)
 			{
+				if (--DrainBufferLimit == 0)
+				{
+					// Ignore buffers added while we drain buffers for this thread.
+					break;
+				}
+
 				if (Writer_DrainBuffer(ThreadId, Buffer))
 				{
 					break;
