Account for complete event fitting in the buffer for scoped events.

If memory tracing is enabled we need to make sure that both the enter uid and the event header fits into the current buffer. Otherwise it's possible that while requesting a new buffer and one has to be allocated, allocation trace events are inserted between the enter uid and the event. The result is that allocation event is incorrectly marked as a scope.

diff --git a/Engine/Source/Runtime/TraceLog/Public/Trace/Detail/LogScope.inl b/Engine/Source/Runtime/TraceLog/Public/Trace/Detail/LogScope.inl
--- a/Engine/Source/Runtime/TraceLog/Public/Trace/Detail/LogScope.inl
+++ b/Engine/Source/Runtime/TraceLog/Public/Trace/Detail/LogScope.inl
@@ -201,7 +201,9 @@ FORCENOINLINE auto FLogScope::ScopedEnter()
 	uint8 EnterUid = uint8(EKnownEventUids::EnterScope << EKnownEventUids::_UidShift);
 
 	FWriteBuffer* Buffer = Writer_GetBuffer();
-	if (UNLIKELY(int32((uint8*)Buffer - Buffer->Cursor) < int32(sizeof(EnterUid))))
+	constexpr int32 EventHeaderSize = (EventType::EventFlags & FEventInfo::Flag_NoSync) ? sizeof(FEventHeader) : sizeof(FEventHeaderSync);
+	constexpr int32 RequiredSize = sizeof(EnterUid) + EventHeaderSize + EventType::GetSize();
+	if (UNLIKELY(int32((uint8*)Buffer - Buffer->Cursor) < RequiredSize))
 	{
 		Buffer = Writer_NextBuffer();
 	}
@@ -221,7 +223,9 @@ FORCENOINLINE auto FLogScope::ScopedStampedEnter()
 	uint64 Stamp;
 
 	FWriteBuffer* Buffer = Writer_GetBuffer();
-	if (UNLIKELY(int32((uint8*)Buffer - Buffer->Cursor) < int32(sizeof(Stamp))))
+	constexpr int32 EventHeaderSize = (EventType::EventFlags & FEventInfo::Flag_NoSync) ? sizeof(FEventHeader) : sizeof(FEventHeaderSync);
+	constexpr int32 RequiredSize = sizeof(Stamp) + EventHeaderSize + EventType::GetSize();
+	if (UNLIKELY(int32((uint8*)Buffer - Buffer->Cursor) < RequiredSize))
 	{
 		Buffer = Writer_NextBuffer();
 	}
