Fix static analysis warnings

diff --git a/Engine/Source/Runtime/TraceLog/Public/Trace/Detail/LogScope.inl b/Engine/Source/Runtime/TraceLog/Public/Trace/Detail/LogScope.inl
--- a/Engine/Source/Runtime/TraceLog/Public/Trace/Detail/LogScope.inl
+++ b/Engine/Source/Runtime/TraceLog/Public/Trace/Detail/LogScope.inl
@@ -201,7 +201,7 @@ FORCENOINLINE auto FLogScope::ScopedEnter()
 	uint8 EnterUid = uint8(EKnownEventUids::EnterScope << EKnownEventUids::_UidShift);
 
 	FWriteBuffer* Buffer = Writer_GetBuffer();
-	constexpr int32 EventHeaderSize = (EventType::EventFlags & FEventInfo::Flag_NoSync) ? sizeof(FEventHeader) : sizeof(FEventHeaderSync);
+	constexpr int32 EventHeaderSize = (EventType::EventFlags & FEventInfo::Flag_NoSync) != 0 ? sizeof(FEventHeader) : sizeof(FEventHeaderSync);
 	constexpr int32 RequiredSize = sizeof(EnterUid) + EventHeaderSize + EventType::GetSize();
 	if (UNLIKELY(int32((uint8*)Buffer - Buffer->Cursor) < RequiredSize))
 	{
@@ -223,7 +223,7 @@ FORCENOINLINE auto FLogScope::ScopedStampedEnter()
 	uint64 Stamp;
 
 	FWriteBuffer* Buffer = Writer_GetBuffer();
-	constexpr int32 EventHeaderSize = (EventType::EventFlags & FEventInfo::Flag_NoSync) ? sizeof(FEventHeader) : sizeof(FEventHeaderSync);
+	constexpr int32 EventHeaderSize = (EventType::EventFlags & FEventInfo::Flag_NoSync) != 0 ? sizeof(FEventHeader) : sizeof(FEventHeaderSync);
 	constexpr int32 RequiredSize = sizeof(Stamp) + EventHeaderSize + EventType::GetSize();
 	if (UNLIKELY(int32((uint8*)Buffer - Buffer->Cursor) < RequiredSize))
 	{
