Clean up TRACE_PRIVATE_CONTROL_ENABLED for platform extensions

#rnx

diff --git a/Engine/Source/Runtime/TraceLog/Private/Trace/Control.cpp b/Engine/Source/Runtime/TraceLog/Private/Trace/Control.cpp
--- a/Engine/Source/Runtime/TraceLog/Private/Trace/Control.cpp
+++ b/Engine/Source/Runtime/TraceLog/Private/Trace/Control.cpp
@@ -286,11 +286,6 @@ void Writer_UpdateControl()
 ////////////////////////////////////////////////////////////////////////////////
 void Writer_InitializeControl()
 {
-#if PLATFORM_SWITCH
-	GControlState = EControlState::Failed;
-	return;
-#endif
-
 	Writer_ControlAddCommand("SendTo", nullptr,
 		[] (void*, uint32 ArgC, ANSICHAR const* const* ArgV)
 		{
