Emit error message if redirect buffer is full

diff --git a/Engine/Source/Runtime/TraceLog/Private/Trace/TlsBuffer.cpp b/Engine/Source/Runtime/TraceLog/Private/Trace/TlsBuffer.cpp
--- a/Engine/Source/Runtime/TraceLog/Private/Trace/TlsBuffer.cpp
+++ b/Engine/Source/Runtime/TraceLog/Private/Trace/TlsBuffer.cpp
@@ -1,5 +1,6 @@
 // Copyright Epic Games, Inc. All Rights Reserved.
 
+#include "Message.h"
 #include "Trace/Config.h"
 
 #if UE_TRACE_ENABLED
@@ -66,6 +67,7 @@ static FWriteBuffer* Writer_NextBufferInternal(FWriteBuffer* CurrentBuffer)
 		// Allocating memory results in so many trace events that there is
 		// insufficient space to store them. We can't allocate more space, because
 		// that would result in more traced events, and so on...
+		UE_TRACE_MESSAGE(WriterError, "Redirect buffer full, trace will be corrupt!");
 		PLATFORM_BREAK();
 	}
 #endif
diff --git a/Engine/Source/Runtime/TraceLog/Public/Trace/Trace.h b/Engine/Source/Runtime/TraceLog/Public/Trace/Trace.h
--- a/Engine/Source/Runtime/TraceLog/Public/Trace/Trace.h
+++ b/Engine/Source/Runtime/TraceLog/Public/Trace/Trace.h
@@ -88,6 +88,7 @@ enum class EMessageType : uint8
 	ListenError,
 	EstablishError,
 	FileOpenError,
+	WriterError,
 };
 	
 struct FMessageEvent
