No point in only forward declaring struct if users need to use it in function later.

diff --git a/Engine/Source/Runtime/TraceLog/Public/Trace/Detail/Channel.h b/Engine/Source/Runtime/TraceLog/Public/Trace/Detail/Channel.h
--- a/Engine/Source/Runtime/TraceLog/Public/Trace/Detail/Channel.h
+++ b/Engine/Source/Runtime/TraceLog/Public/Trace/Detail/Channel.h
@@ -3,6 +3,7 @@
 #pragma once
 
 #include "Trace/Config.h"
+#include "Trace/Trace.h"
 
 #if UE_TRACE_ENABLED
 
@@ -11,13 +12,6 @@
 namespace UE {
 namespace Trace {
 
-struct FChannelInfo
-{
-	const ANSICHAR* Name;
-	const ANSICHAR* Desc;
-	bool bIsEnabled;
-	bool bIsReadOnly;
-};
 
 typedef bool ChannelIterCallback(const FChannelInfo& OutChannelInfo, void*);
 
diff --git a/Engine/Source/Runtime/TraceLog/Public/Trace/Trace.h b/Engine/Source/Runtime/TraceLog/Public/Trace/Trace.h
--- a/Engine/Source/Runtime/TraceLog/Public/Trace/Trace.h
+++ b/Engine/Source/Runtime/TraceLog/Public/Trace/Trace.h
@@ -113,7 +113,13 @@ struct FInitializeDesc
 	OnConnectFunc*	OnConnectionFunc	= nullptr;
 };
 
-struct FChannelInfo;
+struct FChannelInfo
+{
+	const ANSICHAR* Name;
+	const ANSICHAR* Desc;
+	bool bIsEnabled;
+	bool bIsReadOnly;
+};
 
 typedef void*		AllocFunc(SIZE_T, uint32);
 typedef void		FreeFunc(void*, SIZE_T);
