Expose channel ids

Channels have an internal uint32 id. Expose this when enumerating channels, and provide method to toggle channels by id.

diff --git a/Engine/Source/Runtime/TraceLog/Private/Trace/Channel.cpp b/Engine/Source/Runtime/TraceLog/Private/Trace/Channel.cpp
--- a/Engine/Source/Runtime/TraceLog/Private/Trace/Channel.cpp
+++ b/Engine/Source/Runtime/TraceLog/Private/Trace/Channel.cpp
@@ -250,6 +250,30 @@ FChannel* FChannel::FindChannel(const ANSICHAR* ChannelName)
 
 	return nullptr;
 }
+	
+///////////////////////////////////////////////////////////////////////////////
+FChannel* FChannel::FindChannel(FChannelId ChannelId)
+{
+	using namespace Private;
+
+	FChannel* ChannelLists[] =
+	{
+		AtomicLoadAcquire(&GNewChannelList),
+		AtomicLoadAcquire(&GHeadChannel),
+	};
+	for (FChannel* Channel : ChannelLists)
+	{
+		for (; Channel != nullptr; Channel = Channel->Next)
+		{
+			if (Channel->Name.Hash == ChannelId)
+			{
+				return Channel;
+			}
+		}
+	}
+
+	return nullptr;
+}
 
 ///////////////////////////////////////////////////////////////////////////////
 void FChannel::EnumerateChannels(ChannelIterCallback Func, void* User)
@@ -270,6 +294,7 @@ void FChannel::EnumerateChannels(ChannelIterCallback Func, void* User)
 			Info.Desc = Channel->Args.Desc;
 			Info.bIsEnabled = Channel->IsEnabled();
 			Info.bIsReadOnly = Channel->Args.bReadOnly;
+			Info.Id = Channel->Name.Hash;
 			bool Result = Func(Info, User);
 			if (!Result)
 			{
diff --git a/Engine/Source/Runtime/TraceLog/Private/Trace/Trace.cpp b/Engine/Source/Runtime/TraceLog/Private/Trace/Trace.cpp
--- a/Engine/Source/Runtime/TraceLog/Private/Trace/Trace.cpp
+++ b/Engine/Source/Runtime/TraceLog/Private/Trace/Trace.cpp
@@ -194,6 +194,20 @@ void StartWorkerThread()
 {
 	Private::Writer_WorkerCreate();
 }
+	
+////////////////////////////////////////////////////////////////////////////////
+FChannel* FindChannel(const TCHAR* ChannelName)
+{
+	ANSICHAR ChannelNameA[64];
+	ToAnsiCheap(ChannelNameA, ChannelName);
+	return FChannel::FindChannel(ChannelNameA);
+}
+	
+////////////////////////////////////////////////////////////////////////////////
+FChannel* FindChannel(FChannelId ChannelId)
+{
+	return FChannel::FindChannel(ChannelId);
+}
 
 
 ////////////////////////////////////////////////////////////////////////////////
diff --git a/Engine/Source/Runtime/TraceLog/Public/Trace/Detail/Channel.h b/Engine/Source/Runtime/TraceLog/Public/Trace/Detail/Channel.h
--- a/Engine/Source/Runtime/TraceLog/Public/Trace/Detail/Channel.h
+++ b/Engine/Source/Runtime/TraceLog/Public/Trace/Detail/Channel.h
@@ -52,10 +52,13 @@ public:
 	static bool			Toggle(const ANSICHAR* ChannelName, bool bEnabled);
 	static void			ToggleAll(bool bEnabled);
 	static void			PanicDisableAll(); // Disabled channels wont be logged with UE_TRACE_LOG
-	static FChannel*	FindChannel(const ANSICHAR* ChannelName);
+	static FChannel* FindChannel(const ANSICHAR* ChannelName);
+	static FChannel* FindChannel(FChannelId ChannelId);
 	static void			EnumerateChannels(ChannelIterCallback Func, void* User);
-	bool				Toggle(bool bEnabled);
-	bool				IsEnabled() const;
+	TRACELOG_API bool	Toggle(bool bEnabled);
+	bool	IsEnabled() const;
+	bool	IsReadOnly() const { return Args.bReadOnly; };
+	uint32	GetName(const ANSICHAR** OutName) const;
 	explicit			operator bool () const;
 	bool				operator | (const FChannel& Rhs) const;
 
@@ -71,6 +74,17 @@ private:
 	InitArgs			Args;
 };
 
+inline uint32 FChannel::GetName(const ANSICHAR** OutName) const
+{
+	if (OutName != nullptr)
+	{
+		*OutName = Name.Ptr;
+		return Name.Len;
+	}
+	return 0;
+}
+
+
 } // namespace Trace
 } // namespace UE
 
diff --git a/Engine/Source/Runtime/TraceLog/Public/Trace/Trace.h b/Engine/Source/Runtime/TraceLog/Public/Trace/Trace.h
--- a/Engine/Source/Runtime/TraceLog/Public/Trace/Trace.h
+++ b/Engine/Source/Runtime/TraceLog/Public/Trace/Trace.h
@@ -113,14 +113,19 @@ struct FInitializeDesc
 	OnConnectFunc*	OnConnectionFunc	= nullptr;
 };
 
+typedef uint32 FChannelId;
+
 struct FChannelInfo
 {
 	const ANSICHAR* Name;
 	const ANSICHAR* Desc;
+	FChannelId Id;
 	bool bIsEnabled;
 	bool bIsReadOnly;
 };
 
+class FChannel;
+	
 typedef void*		AllocFunc(SIZE_T, uint32);
 typedef void		FreeFunc(void*, SIZE_T);
 typedef void		ChannelIterFunc(const ANSICHAR*, bool, void*);
@@ -160,13 +165,16 @@ UE_TRACE_API bool	SendSnapshotTo(const TCHAR* Host, uint32 Port) UE_TRACE_IMPL(f
 UE_TRACE_API bool	IsTracing() UE_TRACE_IMPL(false);
 UE_TRACE_API bool	IsTracingTo(uint32 (&OutSessionGuid)[4], uint32 (&OutTraceGuid)[4]) UE_TRACE_IMPL(false);
 UE_TRACE_API bool	Stop() UE_TRACE_IMPL(false);
-UE_TRACE_API bool	IsChannel(const TCHAR* ChanneName) UE_TRACE_IMPL(false);
+UE_TRACE_API bool	IsChannel(const TCHAR* ChannelName) UE_TRACE_IMPL(false);
 UE_TRACE_API bool	ToggleChannel(const TCHAR* ChannelName, bool bEnabled) UE_TRACE_IMPL(false);
 UE_TRACE_API void	EnumerateChannels(ChannelIterFunc IterFunc, void* User) UE_TRACE_IMPL();
 UE_TRACE_API void	EnumerateChannels(ChannelIterCallback IterFunc, void* User) UE_TRACE_IMPL();
 UE_TRACE_API void	ThreadRegister(const TCHAR* Name, uint32 SystemId, int32 SortHint) UE_TRACE_IMPL();
 UE_TRACE_API void	ThreadGroupBegin(const TCHAR* Name) UE_TRACE_IMPL();
 UE_TRACE_API void	ThreadGroupEnd() UE_TRACE_IMPL();
+	
+UE_TRACE_API FChannel* FindChannel(const TCHAR* ChannelName) UE_TRACE_IMPL(nullptr);
+UE_TRACE_API FChannel* FindChannel(FChannelId ChannelId) UE_TRACE_IMPL(nullptr);
 
 } // namespace Trace
 } // namespace UE
