Add debug mode to serialized trace packets

By setting UE_TRACE_PACKET_VERIFICATION this mode embeds a serial numbers after each packet. Note that this makes analysis incompatible with traces recorded without this mode. Strictly useful for investigating potential package loss.

diff --git a/Engine/Source/Developer/TraceAnalysis/Private/Analysis/Transport/TidPacketTransport.cpp b/Engine/Source/Developer/TraceAnalysis/Private/Analysis/Transport/TidPacketTransport.cpp
--- a/Engine/Source/Developer/TraceAnalysis/Private/Analysis/Transport/TidPacketTransport.cpp
+++ b/Engine/Source/Developer/TraceAnalysis/Private/Analysis/Transport/TidPacketTransport.cpp
@@ -73,8 +73,14 @@ FTidPacketTransport::EReadPacketResult FTidPacketTransport::ReadPacket()
 													// chance to sample the world at each known sync point.
 	}
 
+#if UE_TRACE_PACKET_VERIFICATION
+	const bool bHasPacketSerial = !!(PacketBase->ThreadId & FTidPacketBase::Verification);
+	FThreadStream* Thread = FindOrAddThread(ThreadId, true);
+#else
 	bool bIsPartial = !!(PacketBase->ThreadId & FTidPacketBase::PartialMarker);
 	FThreadStream* Thread = FindOrAddThread(ThreadId, !bIsPartial);
+#endif
+	
 	if (Thread == nullptr)
 	{
 		return EReadPacketResult::Continue;
@@ -149,6 +155,20 @@ FTidPacketTransport::EReadPacketResult FTidPacketTransport::ReadPacket()
 #endif // UE_TRACE_ANALYSIS_DEBUG
 	}
 
+#if UE_TRACE_PACKET_VERIFICATION
+	if(bHasPacketSerial)
+	{
+		uint64 PacketSerial = *GetPointer<uint64>();
+		if ((LastPacketSerial + 1) != PacketSerial)
+		{
+			UE_LOG(LogCore, Error, TEXT("Found packet with index '%llu' when '%llu` was expected."), PacketSerial, LastPacketSerial);
+			return EReadPacketResult::ReadError;
+		}
+		LastPacketSerial = PacketSerial;
+		FTransport::Advance(sizeof(uint64));
+	}
+#endif
+	
 	return EReadPacketResult::Continue;
 }
 
diff --git a/Engine/Source/Developer/TraceAnalysis/Private/Analysis/Transport/TidPacketTransport.h b/Engine/Source/Developer/TraceAnalysis/Private/Analysis/Transport/TidPacketTransport.h
--- a/Engine/Source/Developer/TraceAnalysis/Private/Analysis/Transport/TidPacketTransport.h
+++ b/Engine/Source/Developer/TraceAnalysis/Private/Analysis/Transport/TidPacketTransport.h
@@ -63,6 +63,9 @@ private:
 
 protected:
 	uint32					Synced = 0x7fff'ffff;
+#if UE_TRACE_PACKET_VERIFICATION
+	uint64					LastPacketSerial = 0;
+#endif
 
 #if UE_TRACE_ANALYSIS_DEBUG
 private:
diff --git a/Engine/Source/Runtime/TraceLog/Private/Trace/Tail.cpp b/Engine/Source/Runtime/TraceLog/Private/Trace/Tail.cpp
--- a/Engine/Source/Runtime/TraceLog/Private/Trace/Tail.cpp
+++ b/Engine/Source/Runtime/TraceLog/Private/Trace/Tail.cpp
@@ -26,6 +26,10 @@ void		Writer_MemoryFree(void*, uint32);
 void		Writer_SendData(uint32, uint8* __restrict, uint32);
 void		Writer_SendDataRaw(const void*, uint32);
 
+#if UE_TRACE_PACKET_VERIFICATION	
+extern uint64 GPacketSerial;
+#endif
+	
 ////////////////////////////////////////////////////////////////////////////////
 // ** See the bottom of this file for an explanation of FPacketRing
 class FPacketRing
@@ -223,9 +227,16 @@ void Writer_TailAppend(uint32 ThreadId, uint8* __restrict Data, uint32 Size)
 	{
 		auto* Packet = GPacketRing.Append<FTidPacket>(Size);
 		Packet->ThreadId = uint16(ThreadId);
+#if UE_TRACE_PACKET_VERIFICATION
+		Packet->ThreadId |= FTidPacketBase::Verification;
+#endif
 		::memcpy(Packet->Data, Data, Size);
 
 		Writer_SendDataRaw(Packet, Packet->PacketSize);
+#if UE_TRACE_PACKET_VERIFICATION
+		const uint64 Serial = GPacketSerial++;
+		Writer_SendDataRaw(&Serial, sizeof(uint64));
+#endif
 		return;
 	}
 
@@ -233,6 +244,9 @@ void Writer_TailAppend(uint32 ThreadId, uint8* __restrict Data, uint32 Size)
 	auto* Packet = GPacketRing.Append<FTidPacketEncoded>(EncodeMaxSize);
 	Packet->ThreadId = uint16(ThreadId);
 	Packet->ThreadId |= FTidPacketBase::EncodedMarker;
+#if UE_TRACE_PACKET_VERIFICATION
+	Packet->ThreadId |= FTidPacketBase::Verification;
+#endif
 	Packet->DecodedSize = uint16(Size);
 
 	uint32 EncodeSize = Encode(Data, Size, Packet->Data, EncodeMaxSize);
@@ -241,6 +255,10 @@ void Writer_TailAppend(uint32 ThreadId, uint8* __restrict Data, uint32 Size)
 	Packet->PacketSize -= uint16(BackUp);
 
 	Writer_SendDataRaw(Packet, Packet->PacketSize);
+#if UE_TRACE_PACKET_VERIFICATION
+	const uint64 Serial = GPacketSerial++;
+	Writer_SendDataRaw(&Serial, sizeof(uint64));
+#endif
 }
 
 ////////////////////////////////////////////////////////////////////////////////
diff --git a/Engine/Source/Runtime/TraceLog/Private/Trace/Writer.cpp b/Engine/Source/Runtime/TraceLog/Private/Trace/Writer.cpp
--- a/Engine/Source/Runtime/TraceLog/Private/Trace/Writer.cpp
+++ b/Engine/Source/Runtime/TraceLog/Private/Trace/Writer.cpp
@@ -85,6 +85,9 @@ TRACELOG_API uint64				GStartCycle;		// = 0;
 TRACELOG_API uint32 volatile	GLogSerial;			// = 0;
 // Counter of calls to Writer_WorkerUpdate to enable regular flushing of output buffers
 static uint32					GUpdateCounter;		// = 0;
+#if UE_TRACE_PACKET_VERIFICATION
+uint64							GPacketSerial = 1;
+#endif
 
 
 
@@ -392,9 +395,17 @@ void Writer_SendData(uint32 ThreadId, uint8* __restrict Data, uint32 Size)
 		Size += sizeof(FTidPacket);
 		auto* Packet = (FTidPacket*)Data;
 		Packet->ThreadId = uint16(ThreadId & FTidPacketBase::ThreadIdMask);
+#if UE_TRACE_PACKET_VERIFICATION
+		Packet->ThreadId |= FTidPacketBase::Verification;
+#endif
 		Packet->PacketSize = uint16(Size);
 
 		Writer_SendDataImpl(Data, Size);
+
+#if UE_TRACE_PACKET_VERIFICATION
+		const uint64 Serial = GPacketSerial++;
+		Writer_SendDataImpl(&Serial, sizeof(uint64));
+#endif
 		return;
 	}
 
@@ -405,11 +416,18 @@ void Writer_SendData(uint32 ThreadId, uint8* __restrict Data, uint32 Size)
 
 	Packet.ThreadId = FTidPacketBase::EncodedMarker;
 	Packet.ThreadId |= uint16(ThreadId & FTidPacketBase::ThreadIdMask);
+#if UE_TRACE_PACKET_VERIFICATION
+	Packet.ThreadId |= FTidPacketBase::Verification;
+#endif
 	Packet.DecodedSize = uint16(Size);
 	Packet.PacketSize = uint16(Encode(Data, Packet.DecodedSize, Packet.Data, sizeof(Packet.Data)));
 	Packet.PacketSize += sizeof(FTidPacketEncoded);
 
 	Writer_SendDataImpl(&Packet, Packet.PacketSize);
+#if UE_TRACE_PACKET_VERIFICATION
+	const uint64 Serial = GPacketSerial++;
+	Writer_SendDataImpl(&Serial, sizeof(uint64));
+#endif
 }
 
 ////////////////////////////////////////////////////////////////////////////////
diff --git a/Engine/Source/Runtime/TraceLog/Public/Trace/Config.h b/Engine/Source/Runtime/TraceLog/Public/Trace/Config.h
--- a/Engine/Source/Runtime/TraceLog/Public/Trace/Config.h
+++ b/Engine/Source/Runtime/TraceLog/Public/Trace/Config.h
@@ -19,3 +19,9 @@
 #if UE_TRACE_ENABLED
 #	define TRACE_PRIVATE_PROTOCOL_7
 #endif
+
+// Enable packet verification. Only useful when looking for transmission bugs. Note that in order to avoid making a new protocol version, enabling
+// this makes existing version 7 traces incompatible.
+#ifndef UE_TRACE_PACKET_VERIFICATION
+#define UE_TRACE_PACKET_VERIFICATION 0
+#endif
\ No newline at end of file
diff --git a/Engine/Source/Runtime/TraceLog/Public/Trace/Detail/Transport.h b/Engine/Source/Runtime/TraceLog/Public/Trace/Detail/Transport.h
--- a/Engine/Source/Runtime/TraceLog/Public/Trace/Detail/Transport.h
+++ b/Engine/Source/Runtime/TraceLog/Public/Trace/Detail/Transport.h
@@ -42,6 +42,7 @@ struct FTidPacketBase
 	{
 		EncodedMarker = 0x8000,
 		PartialMarker = 0x4000, // now unused. fragmented aux-data has an event header
+		Verification  = 0x4000, // when set the packet data is immediately followed by 64-bit verification value, see UE_TRACE_PACKET_VERIFICATION
 		ThreadIdMask  = PartialMarker - 1,
 	};
 
