Fix for trace snapshots on platforms that use internal send buffers

On some platforms we use an internal send buffer to reduce network overheads. When a trace is active and a snapshot was triggerered the send buffer could be flushed while writing the snapshot, making the ongoing trace miss events. To fix this add a forced flush right before creating the snapshot handle.

diff --git a/Engine/Source/Runtime/TraceLog/Private/Trace/Writer.cpp b/Engine/Source/Runtime/TraceLog/Private/Trace/Writer.cpp
--- a/Engine/Source/Runtime/TraceLog/Private/Trace/Writer.cpp
+++ b/Engine/Source/Runtime/TraceLog/Private/Trace/Writer.cpp
@@ -1091,6 +1091,10 @@ bool Writer_WriteSnapshot(const FSnapshotTarget& Target)
 	//  have a limited tolerance to gaps/out-of-order event packets.
 	Writer_WorkerUpdateInternal();
 
+	// Force flush the send buffer so that platforms that use internal send buffers
+	// don't loose data.
+	Writer_FlushSendBuffer();
+
 	{
 		TStashGlobal DataHandle(GDataHandle);
 		TStashGlobal PendingDataHandle(GPendingDataHandle);
