Trace: Avoid allocating a new TLS slot for every thread. Use std::atomic for Counter increments instead.

#rnx

diff --git a/Engine/Source/Runtime/TraceLog/Private/Trace/Writer.cpp b/Engine/Source/Runtime/TraceLog/Private/Trace/Writer.cpp
--- a/Engine/Source/Runtime/TraceLog/Private/Trace/Writer.cpp
+++ b/Engine/Source/Runtime/TraceLog/Private/Trace/Writer.cpp
@@ -162,8 +162,8 @@ uint32 Writer_GetThreadId()
 
 	ThreadOnThreadExit([] () { Writer_EndThreadBuffer(); });
 
-	static uint32 volatile Counter;
-	ThreadId = AtomicAddRelaxed(&Counter, 1u) + ETransportTid::Bias;
+	static std::atomic<uint32> Counter;
+	ThreadId = Counter.fetch_add(1u) + ETransportTid::Bias;
 	return ThreadId;
 }
 
