autortfm: remove instrumentation from TraceLog

- TraceLog is intentionally meant to not depend on Core, but autortfm instrumentation requires Core
- if this causes issues with autortfm code, that stuff needs to be isolated from autortfm, outside of TraceLog itself

diff --git a/Engine/Source/Runtime/TraceLog/TraceLog.Build.cs b/Engine/Source/Runtime/TraceLog/TraceLog.Build.cs
--- a/Engine/Source/Runtime/TraceLog/TraceLog.Build.cs
+++ b/Engine/Source/Runtime/TraceLog/TraceLog.Build.cs
@@ -13,8 +13,6 @@ public class TraceLog : ModuleRules
 		bRequiresImplementModule = false;
 		PublicIncludePathModuleNames.Add("Core");
 
-		bAllowAutoRTFMInstrumentation = true;
-
 		PrivateDefinitions.Add("SUPPRESS_PER_MODULE_INLINE_FILE"); // This module does not use core's standard operator new/delete overloads
 	}
 }
