Disable AutoRTFM here as it does not depend on Core

diff --git a/Engine/Source/Runtime/TraceLog/TraceLog.Build.cs b/Engine/Source/Runtime/TraceLog/TraceLog.Build.cs
--- a/Engine/Source/Runtime/TraceLog/TraceLog.Build.cs
+++ b/Engine/Source/Runtime/TraceLog/TraceLog.Build.cs
@@ -14,5 +14,7 @@ public class TraceLog : ModuleRules
 		PublicIncludePathModuleNames.Add("Core");
 
 		PrivateDefinitions.Add("SUPPRESS_PER_MODULE_INLINE_FILE"); // This module does not use core's standard operator new/delete overloads
+		
+		bDisableAutoRTFMInstrumentation = true;
 	}
 }
