Disable TraceLog statistics when running with thread sanitizer

Synchronizing TraceLog statistics is not critical functionality and will introduce more complexity. Disables statistics reporting when thread sanitizer is enabled.

diff --git a/Engine/Source/Runtime/TraceLog/Public/Trace/Detail/Trace.h b/Engine/Source/Runtime/TraceLog/Public/Trace/Detail/Trace.h
--- a/Engine/Source/Runtime/TraceLog/Public/Trace/Detail/Trace.h
+++ b/Engine/Source/Runtime/TraceLog/Public/Trace/Detail/Trace.h
@@ -18,7 +18,7 @@ class FChannel;
 } // namespace Trace
 } // namespace UE
 
-#define TRACE_PRIVATE_STATISTICS (!UE_BUILD_SHIPPING)
+#define TRACE_PRIVATE_STATISTICS (!(UE_BUILD_SHIPPING || USING_THREAD_SANITISER))
 
 #define TRACE_PRIVATE_CHANNEL_DEFAULT_ARGS false, "None"
 
