Fix TSAN warning in TraceLog

TSAN cannot reason about reading Cursor of the buffer before writing to EtxOffset, which is what releases the buffer to be reaped and eventually recycled. Read the cursor in atomic manner in order to appease TSAN.

diff --git a/Engine/Source/Runtime/TraceLog/Private/Trace/TlsBuffer.cpp b/Engine/Source/Runtime/TraceLog/Private/Trace/TlsBuffer.cpp
--- a/Engine/Source/Runtime/TraceLog/Private/Trace/TlsBuffer.cpp
+++ b/Engine/Source/Runtime/TraceLog/Private/Trace/TlsBuffer.cpp
@@ -317,8 +317,9 @@ void Writer_EndThreadBuffer()
 		return;
 	}
 
-	int32 EtxOffset = int32(PTRINT((uint8*)GTlsWriteBuffer - GTlsWriteBuffer->Cursor));
-	AtomicStoreRelaxed(&(GTlsWriteBuffer->EtxOffset), EtxOffset);
+	const PTRINT CurrentCursor = AtomicLoadRelaxed(GTlsWriteBuffer->Cursor);
+	int32 EtxOffset = int32(PTRINT((uint8*)GTlsWriteBuffer - CurrentCursor));
+	AtomicStoreRelease(&(GTlsWriteBuffer->EtxOffset), EtxOffset);
 }
 
 } // namespace Private
