Do not reset bytes traced on new connection

diff --git a/Engine/Source/Runtime/TraceLog/Private/Trace/Writer.cpp b/Engine/Source/Runtime/TraceLog/Private/Trace/Writer.cpp
--- a/Engine/Source/Runtime/TraceLog/Private/Trace/Writer.cpp
+++ b/Engine/Source/Runtime/TraceLog/Private/Trace/Writer.cpp
@@ -576,7 +576,6 @@ static bool Writer_UpdateConnection()
 
 	// Reset statistics.
 	GTraceStatistics.BytesSent = 0;
-	GTraceStatistics.BytesTraced = 0;
 
 	// The first events we will send are ones that describe the trace's events
 	FEventNode::OnConnect();
