Fixing a deadlock that happens inside WriteFile when connecting Unreal Insight to a running game on Steam Deck.
When creating a connection with WSASocket sending data using WSASend instead of WriteFile do not cause the issue on the Stem Deck.

diff --git a/Engine/Source/Runtime/TraceLog/Private/Trace/Detail/Windows/WindowsTrace.cpp b/Engine/Source/Runtime/TraceLog/Private/Trace/Detail/Windows/WindowsTrace.cpp
--- a/Engine/Source/Runtime/TraceLog/Private/Trace/Detail/Windows/WindowsTrace.cpp
+++ b/Engine/Source/Runtime/TraceLog/Private/Trace/Detail/Windows/WindowsTrace.cpp
@@ -18,6 +18,8 @@ namespace UE {
 namespace Trace {
 namespace Private {
 
+extern volatile bool GIsUsingSocket;
+
 struct FSetLastErrorScope
 {
 	FSetLastErrorScope(DWORD InError)
@@ -233,12 +235,29 @@ bool TcpSocketHasData(UPTRINT Socket)
 ////////////////////////////////////////////////////////////////////////////////
 bool IoWrite(UPTRINT Handle, const void* Data, uint32 Size)
 {
-	HANDLE Inner = HANDLE(Handle - 1);
-
 	DWORD BytesWritten = 0;
-	if (!WriteFile(Inner, (const char*)Data, Size, &BytesWritten, nullptr))
+
+	//Temporary workaround to differentiate between a file and a socket due to some emulation platforms issues
+	//WriteFile is deadlocking inside the function call while WSASend is working as expected, we can revert using WriteFile once the underlying issue is fixed
+	if (GIsUsingSocket)
+	{
+		SOCKET InnerSocket = (SOCKET)(Handle - 1);
+		WSABUF Buffer;
+		Buffer.buf = (CHAR*)Data;
+		Buffer.len = Size;
+		int ret = WSASend(InnerSocket, &Buffer, 1, &BytesWritten, 0, nullptr, nullptr);
+		if (ret != 0)
+		{
+			return false;
+		}
+	}
+	else
 	{
-		return false;
+		HANDLE Inner = HANDLE(Handle - 1);
+		if (!WriteFile(Inner, (const char*)Data, Size, &BytesWritten, nullptr))
+		{
+			return false;
+		}
 	}
 
 	return (BytesWritten == Size);
diff --git a/Engine/Source/Runtime/TraceLog/Private/Trace/Writer.cpp b/Engine/Source/Runtime/TraceLog/Private/Trace/Writer.cpp
--- a/Engine/Source/Runtime/TraceLog/Private/Trace/Writer.cpp
+++ b/Engine/Source/Runtime/TraceLog/Private/Trace/Writer.cpp
@@ -80,6 +80,8 @@ UE_TRACE_EVENT_END()
 
 ////////////////////////////////////////////////////////////////////////////////
 static volatile bool			GInitialized;		// = false;
+//Temporary workaround to differentiate between a file and a socket due to some emulation platforms issues
+volatile bool					GIsUsingSocket 		   = false;
 FStatistics						GTraceStatistics;	// = {};
 TRACELOG_API uint64				GStartCycle;		// = 0;
 TRACELOG_API uint32 volatile	GLogSerial;			// = 0;
@@ -938,6 +940,8 @@ bool Writer_SendTo(const ANSICHAR* Host, uint32 Flags, uint32 Port)
 		return false;
 	}
 
+	GIsUsingSocket = true;
+
 	DataHandle = Writer_PackSendFlags(DataHandle, Flags);
 	if (!DataHandle)
 	{
@@ -972,6 +976,8 @@ bool Writer_WriteTo(const ANSICHAR* Path, uint32 Flags)
 		return false;
 	}
 
+	GIsUsingSocket = false;
+
 	AtomicStoreRelaxed(&GPendingDataHandle, DataHandle);
 	return true;
 }
