Add a simple function pair that can be used to throttle execution

#rnx

diff --git a/Engine/Source/Runtime/TraceLog/Private/Trace/Platform.h b/Engine/Source/Runtime/TraceLog/Private/Trace/Platform.h
--- a/Engine/Source/Runtime/TraceLog/Private/Trace/Platform.h
+++ b/Engine/Source/Runtime/TraceLog/Private/Trace/Platform.h
@@ -38,6 +38,25 @@ UPTRINT	FileOpen(const ANSICHAR* Path);
 int32	GetLastErrorCode();
 bool	GetErrorMessage(char* OutBuffer, uint32 BufferSize, int32 ErrorCode);	
 
+////////////////////////////////////////////////////////////////////////////////
+#if !defined(TRACE_PRIVATE_HAS_THROTTLE)
+#	define	TRACE_PRIVATE_HAS_THROTTLE	0
+#endif
+
+#if TRACE_PRIVATE_HAS_THROTTLE
+#	define	THROTTLE_IMPL(x)	;
+#	define	THROTTLE_INL
+#else
+#	define	THROTTLE_IMPL(x)	{x}
+#	define	THROTTLE_INL		inline
+#endif
+
+THROTTLE_INL int32	ThreadThrottle()		THROTTLE_IMPL(return 0;)
+THROTTLE_INL void	ThreadUnthrottle(int32)	THROTTLE_IMPL()
+
+#undef THROTTLE_INL
+#undef THROTTLE_IMPL
+
 } // namespace Private
 } // namespace Trace
 } // namespace UE
