Throttle waiting for shared buffers to allow progress

On some platforms, the list of active threads is only modified if a
higher-priority thread wakes up. If no thread in that list ever enters a
wait/stop state and progress is blocked, a deadlock occurs. The yield doesn't
help as it just moves a thread to the back of the queue. Everything gets a
slice of the CPU pie if we drop to the lowest priority thread (even those that
don't yield as the scheduler in question timeslices lowest priority threads).

#rnx

diff --git a/Engine/Source/Runtime/TraceLog/Private/Trace/Important/SharedBuffer.cpp b/Engine/Source/Runtime/TraceLog/Private/Trace/Important/SharedBuffer.cpp
--- a/Engine/Source/Runtime/TraceLog/Private/Trace/Important/SharedBuffer.cpp
+++ b/Engine/Source/Runtime/TraceLog/Private/Trace/Important/SharedBuffer.cpp
@@ -5,8 +5,10 @@
 #if UE_TRACE_ENABLED
 
 #include "CoreTypes.h"
+#include "Misc/ScopeExit.h"
 #include "Trace/Detail/Atomic.h"
 #include "Trace/Detail/Important/ImportantLogScope.inl"
+#include "Trace/Platform.h"
 
 #ifndef PLATFORM_TRACE_WRITER_BUFFER_SIZE
 #define PLATFORM_TRACE_WRITER_BUFFER_SIZE				( 1*1024 )
@@ -73,6 +75,9 @@ FNextSharedBuffer Writer_NextSharedBuffer(FSharedBuffer* Buffer, int32 RegionSta
 		}
 		else
 		{
+			int32 ThrottleRestore = ThreadThrottle();
+			ON_SCOPE_EXIT { ThreadUnthrottle(ThrottleRestore); };
+
 			// Another thread is already allocating the next buffer, wait for that to complete
 			for (;; PlatformYield())
 			{
@@ -141,11 +146,27 @@ static void Writer_RetireSharedBuffer()
 ////////////////////////////////////////////////////////////////////////////////
 void Writer_UpdateSharedBuffers()
 {
+	int32 ThrottleRestore = -1;
+	ON_SCOPE_EXIT {
+		if (ThrottleRestore >= 0)
+		{
+			ThreadUnthrottle(ThrottleRestore);
+		}
+	};
+
+	auto ApplyThrottle = [&ThrottleRestore] () {
+		if (ThrottleRestore < 0)
+		{
+			ThrottleRestore = ThreadThrottle();
+		}
+	};
+
 	FSharedBuffer* HeadBuffer = AtomicLoadAcquire(&GSharedBuffer);
-	while (true)
+	for (;; PlatformYield())
 	{
 		if (GTailBuffer != HeadBuffer)
 		{
+			ApplyThrottle();
 			Writer_RetireSharedBuffer();
 			continue;
 		}
@@ -153,6 +174,7 @@ void Writer_UpdateSharedBuffers()
 		int32 Cursor = AtomicLoadAcquire(&(HeadBuffer->Cursor));
 		if ((Cursor + 1) & FSharedBuffer::RefInit)
 		{
+			ApplyThrottle();
 			continue;
 		}
 
