Undid a previous attempt to fix a deadlock

diff --git a/Engine/Source/Runtime/TraceLog/Private/Trace/Important/SharedBuffer.cpp b/Engine/Source/Runtime/TraceLog/Private/Trace/Important/SharedBuffer.cpp
--- a/Engine/Source/Runtime/TraceLog/Private/Trace/Important/SharedBuffer.cpp
+++ b/Engine/Source/Runtime/TraceLog/Private/Trace/Important/SharedBuffer.cpp
@@ -10,10 +10,6 @@
 #include "Trace/Detail/Important/ImportantLogScope.inl"
 #include "Trace/Platform.h"
 
-#ifndef PLATFORM_TRACE_WRITER_BUFFER_SIZE
-#define PLATFORM_TRACE_WRITER_BUFFER_SIZE				( 1*1024 )
-#endif // PLATFORM_TRACE_WRITER_BUFFER_SIZE
-
 namespace UE {
 namespace Trace {
 namespace Private {
@@ -28,7 +24,7 @@ static FSharedBuffer	GNullSharedBuffer	= { 0, FSharedBuffer::RefInit };
 FSharedBuffer* volatile GSharedBuffer		= &GNullSharedBuffer;
 static FSharedBuffer*	GTailBuffer;		// = nullptr
 static uint32			GTailPreSent;		// = 0
-static const uint32		GBlockSize			= PLATFORM_TRACE_WRITER_BUFFER_SIZE;		// Block size must be a power of two!
+static const uint32		GBlockSize			= 1 << 10;		// Block size must be a power of two!
 extern FStatistics		GTraceStatistics;
 
 ////////////////////////////////////////////////////////////////////////////////
