Enable Trace in shipping: TraceLog

Introduce a new configuration (UE_TRACE_ENABLED_SHIPPING_EXPERIMENTAL) which optionally enables select parts of the TraceLog library in shipping configurations. In this configuration, only a subset of trace events and logs are active. These needs to be de specificially for shipping using a new set of macros. Addtionally this change adds the ability to configure other parts of TraceLog functionality such as the socket based control layer, tracing to sockets or files and if important events are allowed. This is currently an experimental feature.

diff --git a/Engine/Source/Runtime/Core/Private/ProfilingDebugging/PlatformEvents.cpp b/Engine/Source/Runtime/Core/Private/ProfilingDebugging/PlatformEvents.cpp
--- a/Engine/Source/Runtime/Core/Private/ProfilingDebugging/PlatformEvents.cpp
+++ b/Engine/Source/Runtime/Core/Private/ProfilingDebugging/PlatformEvents.cpp
@@ -109,11 +109,11 @@ void FPlatformEventsTrace::OnTraceChannelUpdated(const FString& ChannelName, boo
 
 void FPlatformEventsTrace::PostInit()
 {
-	if (TRACE_PRIVATE_CHANNELEXPR_IS_ENABLED(ContextSwitchChannel)) //-V517
+	if (UE_TRACE_CHANNELEXPR_IS_ENABLED(ContextSwitchChannel)) //-V517
 	{
 		Enable(FPlatformEventsTrace::EEventType::ContextSwitch);
 	}
-	else if (TRACE_PRIVATE_CHANNELEXPR_IS_ENABLED(StackSamplingChannel))
+	else if (UE_TRACE_CHANNELEXPR_IS_ENABLED(StackSamplingChannel))
 	{
 		Enable(FPlatformEventsTrace::EEventType::StackSampling);
 	}
diff --git a/Engine/Source/Runtime/TraceLog/Private/Trace/BlockPool.cpp b/Engine/Source/Runtime/TraceLog/Private/Trace/BlockPool.cpp
--- a/Engine/Source/Runtime/TraceLog/Private/Trace/BlockPool.cpp
+++ b/Engine/Source/Runtime/TraceLog/Private/Trace/BlockPool.cpp
@@ -2,7 +2,7 @@
 
 #include "Trace/Config.h"
 
-#if UE_TRACE_ENABLED
+#if TRACE_PRIVATE_MINIMAL_ENABLED
 
 #include "Platform.h"
 #include "Trace/Detail/Atomic.h"
@@ -190,4 +190,4 @@ void Writer_ShutdownPool()
 } // namespace Trace
 } // namespace UE
 
-#endif // UE_TRACE_ENABLED
+#endif // TRACE_PRIVATE_MINIMAL_ENABLED
diff --git a/Engine/Source/Runtime/TraceLog/Private/Trace/Channel.cpp b/Engine/Source/Runtime/TraceLog/Private/Trace/Channel.cpp
--- a/Engine/Source/Runtime/TraceLog/Private/Trace/Channel.cpp
+++ b/Engine/Source/Runtime/TraceLog/Private/Trace/Channel.cpp
@@ -9,7 +9,7 @@
 
 #include <ctype.h>
 
-#if UE_TRACE_ENABLED
+#if TRACE_PRIVATE_MINIMAL_ENABLED
 
 namespace UE {
 namespace Trace {
@@ -25,17 +25,17 @@ static FTraceChannel	TraceLogChannelDetail;
 FChannel&				TraceLogChannel			= TraceLogChannelDetail;
 
 ///////////////////////////////////////////////////////////////////////////////
-UE_TRACE_EVENT_BEGIN(Trace, ChannelAnnounce, NoSync|Important)
-	UE_TRACE_EVENT_FIELD(uint32, Id)
-	UE_TRACE_EVENT_FIELD(bool, IsEnabled)
-	UE_TRACE_EVENT_FIELD(bool, ReadOnly)
-	UE_TRACE_EVENT_FIELD(AnsiString, Name)
-UE_TRACE_EVENT_END()
-
-UE_TRACE_EVENT_BEGIN(Trace, ChannelToggle, NoSync|Important)
-	UE_TRACE_EVENT_FIELD(uint32, Id)
-	UE_TRACE_EVENT_FIELD(bool, IsEnabled)
-UE_TRACE_EVENT_END()
+UE_TRACE_MINIMAL_EVENT_BEGIN(Trace, ChannelAnnounce, NoSync|Important)
+	UE_TRACE_MINIMAL_EVENT_FIELD(uint32, Id)
+	UE_TRACE_MINIMAL_EVENT_FIELD(bool, IsEnabled)
+	UE_TRACE_MINIMAL_EVENT_FIELD(bool, ReadOnly)
+	UE_TRACE_MINIMAL_EVENT_FIELD(AnsiString, Name)
+UE_TRACE_MINIMAL_EVENT_END()
+
+UE_TRACE_MINIMAL_EVENT_BEGIN(Trace, ChannelToggle, NoSync|Important)
+	UE_TRACE_MINIMAL_EVENT_FIELD(uint32, Id)
+	UE_TRACE_MINIMAL_EVENT_FIELD(bool, IsEnabled)
+UE_TRACE_MINIMAL_EVENT_END()
 
 ///////////////////////////////////////////////////////////////////////////////
 static FChannel* volatile	GHeadChannel;			// = nullptr;
@@ -167,7 +167,7 @@ void FChannel::Setup(const ANSICHAR* InChannelName, const InitArgs& InArgs)
 ///////////////////////////////////////////////////////////////////////////////
 void FChannel::Announce() const
 {
-	UE_TRACE_LOG(Trace, ChannelAnnounce, TraceLogChannel, Name.Len * sizeof(ANSICHAR))
+	UE_TRACE_MINIMAL_LOG(Trace, ChannelAnnounce, TraceLogChannel, Name.Len * sizeof(ANSICHAR))
 		<< ChannelAnnounce.Id(Name.Hash)
 		<< ChannelAnnounce.IsEnabled(IsEnabled())
 		<< ChannelAnnounce.ReadOnly(Args.bReadOnly)
@@ -310,7 +310,7 @@ bool FChannel::Toggle(bool bEnabled)
 	using namespace Private;
 	AtomicStoreRelaxed(&Enabled, bEnabled ? 1 : -1);
 
-	UE_TRACE_LOG(Trace, ChannelToggle, TraceLogChannel)
+	UE_TRACE_MINIMAL_LOG(Trace, ChannelToggle, TraceLogChannel)
 		<< ChannelToggle.Id(Name.Hash)
 		<< ChannelToggle.IsEnabled(IsEnabled());
 
@@ -330,4 +330,4 @@ bool FChannel::Toggle(const ANSICHAR* ChannelName, bool bEnabled)
 } // namespace Trace
 } // namespace UE
 
-#endif // UE_TRACE_ENABLED
+#endif // TRACE_PRIVATE_MINIMAL_ENABLED
diff --git a/Engine/Source/Runtime/TraceLog/Private/Trace/Control.cpp b/Engine/Source/Runtime/TraceLog/Private/Trace/Control.cpp
--- a/Engine/Source/Runtime/TraceLog/Private/Trace/Control.cpp
+++ b/Engine/Source/Runtime/TraceLog/Private/Trace/Control.cpp
@@ -2,8 +2,6 @@
 
 #include "Trace/Config.h"
 
-#if UE_TRACE_ENABLED
-
 #include "Trace/Platform.h"
 #include "Trace/Message.h"
 #include "Trace/Detail/Channel.h"
@@ -17,7 +15,7 @@ namespace UE {
 namespace Trace {
 namespace Private {
 
-#if !defined(TRACE_PRIVATE_CONTROL_ENABLED) || TRACE_PRIVATE_CONTROL_ENABLED
+#if TRACE_PRIVATE_ALLOW_TCP_CONTROL
 
 ////////////////////////////////////////////////////////////////////////////////
 bool	Writer_SendTo(const ANSICHAR*, uint32=0, uint32=0);
@@ -343,11 +341,9 @@ void	Writer_InitializeControl()	{}
 void	Writer_ShutdownControl()	{}
 void	Writer_UpdateControl()		{}
 uint32	Writer_GetControlPort()		{ return ~0u; }
-
-#endif // TRACE_PRIVATE_CONTROL_ENABLED
+	
+#endif // TRACE_PRIVATE_ALLOW_TCP_CONTROL
 
 } // namespace Private
 } // namespace Trace
 } // namespace UE
-
-#endif // UE_TRACE_ENABLED
diff --git a/Engine/Source/Runtime/TraceLog/Private/Trace/Detail/Android/AndroidTrace.cpp b/Engine/Source/Runtime/TraceLog/Private/Trace/Detail/Android/AndroidTrace.cpp
--- a/Engine/Source/Runtime/TraceLog/Private/Trace/Detail/Android/AndroidTrace.cpp
+++ b/Engine/Source/Runtime/TraceLog/Private/Trace/Detail/Android/AndroidTrace.cpp
@@ -2,7 +2,7 @@
 
 #include "Trace/Config.h"
 
-#if UE_TRACE_ENABLED && PLATFORM_ANDROID
+#if TRACE_PRIVATE_MINIMAL_ENABLED && PLATFORM_ANDROID
 
 #include <arpa/inet.h>
 #include <string.h>
@@ -243,4 +243,4 @@ bool GetErrorMessage(char* OutBuffer, uint32 BufferSize, int32 ErrorCode)
 } // namespace Trace
 } // namespace UE
 
-#endif // UE_TRACE_ENABLED
+#endif // TRACE_PRIVATE_MINIMAL_ENABLED
diff --git a/Engine/Source/Runtime/TraceLog/Private/Trace/Detail/Apple/AppleTrace.cpp b/Engine/Source/Runtime/TraceLog/Private/Trace/Detail/Apple/AppleTrace.cpp
--- a/Engine/Source/Runtime/TraceLog/Private/Trace/Detail/Apple/AppleTrace.cpp
+++ b/Engine/Source/Runtime/TraceLog/Private/Trace/Detail/Apple/AppleTrace.cpp
@@ -2,7 +2,7 @@
 
 #include "Trace/Config.h"
 
-#if UE_TRACE_ENABLED && PLATFORM_APPLE
+#if TRACE_PRIVATE_MINIMAL_ENABLED && PLATFORM_APPLE
 
 #include <arpa/inet.h>
 #include <string.h>
@@ -275,4 +275,4 @@ bool GetErrorMessage(char* OutBuffer, uint32 BufferSize, int32 ErrorCode)
 } // namespace Trace
 } // namespace UE
 
-#endif // UE_TRACE_ENABLED
+#endif // TRACE_PRIVATE_MINIMAL_ENABLED
diff --git a/Engine/Source/Runtime/TraceLog/Private/Trace/Detail/Unix/UnixTrace.cpp b/Engine/Source/Runtime/TraceLog/Private/Trace/Detail/Unix/UnixTrace.cpp
--- a/Engine/Source/Runtime/TraceLog/Private/Trace/Detail/Unix/UnixTrace.cpp
+++ b/Engine/Source/Runtime/TraceLog/Private/Trace/Detail/Unix/UnixTrace.cpp
@@ -2,7 +2,7 @@
 
 #include "Trace/Config.h"
 
-#if UE_TRACE_ENABLED && PLATFORM_UNIX
+#if TRACE_PRIVATE_MINIMAL_ENABLED && PLATFORM_UNIX
 
 #include <cstring>
 #include <errno.h>
@@ -274,4 +274,4 @@ bool GetErrorMessage(char* OutBuffer, uint32 BufferSize, int32 ErrorCode)
 } // namespace Trace
 } // namespace UE
 
-#endif // UE_TRACE_ENABLED
+#endif // TRACE_PRIVATE_MINIMAL_ENABLED
diff --git a/Engine/Source/Runtime/TraceLog/Private/Trace/Detail/Windows/WindowsTrace.cpp b/Engine/Source/Runtime/TraceLog/Private/Trace/Detail/Windows/WindowsTrace.cpp
--- a/Engine/Source/Runtime/TraceLog/Private/Trace/Detail/Windows/WindowsTrace.cpp
+++ b/Engine/Source/Runtime/TraceLog/Private/Trace/Detail/Windows/WindowsTrace.cpp
@@ -2,7 +2,7 @@
 
 #include "Trace/Config.h"
 
-#if UE_TRACE_ENABLED && PLATFORM_WINDOWS
+#if TRACE_PRIVATE_MINIMAL_ENABLED && PLATFORM_WINDOWS
 
 #include "Windows/AllowWindowsPlatformTypes.h"
 #	define _WINSOCK_DEPRECATED_NO_WARNINGS
@@ -302,4 +302,4 @@ bool GetErrorMessage(char* OutBuffer, uint32 BufferSize, int32 ErrorCode)
 
 #pragma warning(pop)
 
-#endif // UE_TRACE_ENABLED
+#endif // TRACE_PRIVATE_MINIMAL_ENABLED
diff --git a/Engine/Source/Runtime/TraceLog/Private/Trace/EventNode.cpp b/Engine/Source/Runtime/TraceLog/Private/Trace/EventNode.cpp
--- a/Engine/Source/Runtime/TraceLog/Private/Trace/EventNode.cpp
+++ b/Engine/Source/Runtime/TraceLog/Private/Trace/EventNode.cpp
@@ -2,7 +2,7 @@
 
 #include "Trace/Detail/EventNode.h"
 
-#if UE_TRACE_ENABLED
+#if TRACE_PRIVATE_MINIMAL_ENABLED
 
 #include "Trace/Detail/Atomic.h"
 #include "Trace/Detail/LogScope.inl"
@@ -227,5 +227,5 @@ void FEventNode::OnConnect()
 } // namespace Trace
 } // namespace UE
 
-#endif // UE_TRACE_ENABLED
+#endif // TRACE_PRIVATE_MINIMAL_ENABLED
 
diff --git a/Engine/Source/Runtime/TraceLog/Private/Trace/Field.cpp b/Engine/Source/Runtime/TraceLog/Private/Trace/Field.cpp
--- a/Engine/Source/Runtime/TraceLog/Private/Trace/Field.cpp
+++ b/Engine/Source/Runtime/TraceLog/Private/Trace/Field.cpp
@@ -2,7 +2,7 @@
 
 #include "Trace/Detail/Field.h" // :(
 
-#if UE_TRACE_ENABLED
+#if TRACE_PRIVATE_MINIMAL_ENABLED
 
 #include "Trace/Detail/Writer.inl"
 
@@ -139,4 +139,4 @@ void Field_WriteStringWide(uint32 Index, const WIDECHAR* String, int32 Length)
 } // namespace Trace
 } // namespace UE
 
-#endif // UE_TRACE_ENABLED
+#endif // TRACE_PRIVATE_MINIMAL_ENABLED
diff --git a/Engine/Source/Runtime/TraceLog/Private/Trace/Important/Cache.cpp b/Engine/Source/Runtime/TraceLog/Private/Trace/Important/Cache.cpp
--- a/Engine/Source/Runtime/TraceLog/Private/Trace/Important/Cache.cpp
+++ b/Engine/Source/Runtime/TraceLog/Private/Trace/Important/Cache.cpp
@@ -2,7 +2,7 @@
 
 #include "Trace/Config.h"
 
-#if UE_TRACE_ENABLED
+#if TRACE_PRIVATE_MINIMAL_ENABLED && TRACE_PRIVATE_ALLOW_IMPORTANTS
 
 #include "Trace/Detail/Protocol.h"
 #include "Trace/Detail/Transport.h"
@@ -188,4 +188,4 @@ void Writer_ShutdownCache()
 } // namespace Trace
 } // namespace UE
 
-#endif // UE_TRACE_ENABLED
+#endif // TRACE_PRIVATE_MINIMAL_ENABLED && TRACE_PRIVATE_ALLOW_IMPORTANTS
diff --git a/Engine/Source/Runtime/TraceLog/Private/Trace/Important/SharedBuffer.cpp b/Engine/Source/Runtime/TraceLog/Private/Trace/Important/SharedBuffer.cpp
--- a/Engine/Source/Runtime/TraceLog/Private/Trace/Important/SharedBuffer.cpp
+++ b/Engine/Source/Runtime/TraceLog/Private/Trace/Important/SharedBuffer.cpp
@@ -2,7 +2,7 @@
 
 #include "Trace/Detail/Important/SharedBuffer.h"
 
-#if UE_TRACE_ENABLED
+#if TRACE_PRIVATE_MINIMAL_ENABLED && TRACE_PRIVATE_ALLOW_IMPORTANTS
 
 #include "CoreTypes.h"
 #include "Misc/ScopeExit.h"
@@ -213,4 +213,4 @@ void Writer_ShutdownSharedBuffers()
 } // namespace Trace
 } // namespace UE
 
-#endif // UE_TRACE_ENABLED
+#endif // TRACE_PRIVATE_MINIMAL_ENABLED && UE_TRACE_ALLOW_IMPORTANTS
diff --git a/Engine/Source/Runtime/TraceLog/Private/Trace/Message.cpp b/Engine/Source/Runtime/TraceLog/Private/Trace/Message.cpp
--- a/Engine/Source/Runtime/TraceLog/Private/Trace/Message.cpp
+++ b/Engine/Source/Runtime/TraceLog/Private/Trace/Message.cpp
@@ -2,7 +2,7 @@
 
 #include "Message.h"
 
-#if UE_TRACE_ENABLED
+#if TRACE_PRIVATE_MINIMAL_ENABLED
 
 #include "Trace/Trace.h"
 
@@ -33,4 +33,4 @@ void Message_Send(EMessageType Type, const char* TypeStr, const char* Descriptio
 
 } } } // namespace UE::Trace::Private
 
-#endif // UE_TRACE_ENABLED
\ No newline at end of file
+#endif // TRACE_PRIVATE_MINIMAL_ENABLED
\ No newline at end of file
diff --git a/Engine/Source/Runtime/TraceLog/Private/Trace/Message.h b/Engine/Source/Runtime/TraceLog/Private/Trace/Message.h
--- a/Engine/Source/Runtime/TraceLog/Private/Trace/Message.h
+++ b/Engine/Source/Runtime/TraceLog/Private/Trace/Message.h
@@ -5,7 +5,7 @@
 #include "Trace/Config.h"
 #include "Trace/Trace.h"
 
-#if UE_TRACE_ENABLED
+#if TRACE_PRIVATE_MINIMAL_ENABLED
 
 namespace UE {
 namespace Trace {
@@ -66,4 +66,4 @@ void Message_Send(EMessageType Type, const char* TypeStr, const char* Descriptio
 
 } } } // namespace UE::Trace::Private
 
-#endif // UE_TRACE_ENABLED
+#endif // TRACE_PRIVATE_MINIMAL_ENABLED
diff --git a/Engine/Source/Runtime/TraceLog/Private/Trace/Platform.h b/Engine/Source/Runtime/TraceLog/Private/Trace/Platform.h
--- a/Engine/Source/Runtime/TraceLog/Private/Trace/Platform.h
+++ b/Engine/Source/Runtime/TraceLog/Private/Trace/Platform.h
@@ -4,7 +4,7 @@
 
 #include "Trace/Config.h"
 
-#if UE_TRACE_ENABLED
+#if TRACE_PRIVATE_MINIMAL_ENABLED
 
 namespace UE {
 namespace Trace {
@@ -61,4 +61,4 @@ THROTTLE_INL void	ThreadUnthrottle(int32)	THROTTLE_IMPL()
 } // namespace Trace
 } // namespace UE
 
-#endif // UE_TRACE_ENABLED
+#endif // TRACE_PRIVATE_MINIMAL_ENABLED
diff --git a/Engine/Source/Runtime/TraceLog/Private/Trace/Tail.cpp b/Engine/Source/Runtime/TraceLog/Private/Trace/Tail.cpp
--- a/Engine/Source/Runtime/TraceLog/Private/Trace/Tail.cpp
+++ b/Engine/Source/Runtime/TraceLog/Private/Trace/Tail.cpp
@@ -2,7 +2,7 @@
 
 #include "Trace/Config.h"
 
-#if UE_TRACE_ENABLED
+#if TRACE_PRIVATE_MINIMAL_ENABLED
 
 #include "Trace/Detail/Protocol.h"
 #include "Trace/Detail/Transport.h"
@@ -350,7 +350,7 @@ static void StressRingPacket()
 } // namespace Trace
 } // namespace UE
 
-#endif // UE_TRACE_ENABLED
+#endif // TRACE_PRIVATE_MINIMAL_ENABLED
 
 /*
 
diff --git a/Engine/Source/Runtime/TraceLog/Private/Trace/TlsBuffer.cpp b/Engine/Source/Runtime/TraceLog/Private/Trace/TlsBuffer.cpp
--- a/Engine/Source/Runtime/TraceLog/Private/Trace/TlsBuffer.cpp
+++ b/Engine/Source/Runtime/TraceLog/Private/Trace/TlsBuffer.cpp
@@ -3,7 +3,7 @@
 #include "Message.h"
 #include "Trace/Config.h"
 
-#if UE_TRACE_ENABLED
+#if TRACE_PRIVATE_MINIMAL_ENABLED
 
 #include "HAL/Platform.h" // for PLATFORM_BREAK
 #include "Platform.h"
@@ -27,9 +27,9 @@ extern FStatistics	GTraceStatistics;
 
 
 ////////////////////////////////////////////////////////////////////////////////
-UE_TRACE_EVENT_BEGIN($Trace, ThreadTiming, NoSync)
-	UE_TRACE_EVENT_FIELD(uint64, BaseTimestamp)
-UE_TRACE_EVENT_END()
+UE_TRACE_MINIMAL_EVENT_BEGIN($Trace, ThreadTiming, NoSync)
+	UE_TRACE_MINIMAL_EVENT_FIELD(uint64, BaseTimestamp)
+UE_TRACE_MINIMAL_EVENT_END()
 
 
 
@@ -96,7 +96,7 @@ static FWriteBuffer* Writer_NextBufferInternal(FWriteBuffer* CurrentBuffer)
 		NextBuffer->ThreadId = uint16(Writer_GetThreadId());
 		NextBuffer->PrevTimestamp = TimeGetTimestamp();
 
-		UE_TRACE_LOG($Trace, ThreadTiming, TraceLogChannel)
+		UE_TRACE_MINIMAL_LOG($Trace, ThreadTiming, TraceLogChannel)
 			<< ThreadTiming.BaseTimestamp(NextBuffer->PrevTimestamp - GStartCycle);
 
 		// Add this next buffer to the active list.
@@ -326,4 +326,4 @@ void Writer_EndThreadBuffer()
 } // namespace Trace
 } // namespace UE
 
-#endif // UE_TRACE_ENABLED
+#endif // TRACE_PRIVATE_MINIMAL_ENABLED
diff --git a/Engine/Source/Runtime/TraceLog/Private/Trace/Trace.cpp b/Engine/Source/Runtime/TraceLog/Private/Trace/Trace.cpp
--- a/Engine/Source/Runtime/TraceLog/Private/Trace/Trace.cpp
+++ b/Engine/Source/Runtime/TraceLog/Private/Trace/Trace.cpp
@@ -2,7 +2,7 @@
 
 #include "Trace/Trace.inl" // should be Config.h :(
 
-#if UE_TRACE_ENABLED
+#if TRACE_PRIVATE_MINIMAL_ENABLED
 
 #include "Misc/CString.h"
 #include "Trace/Detail/Channel.h"
@@ -211,22 +211,23 @@ FChannel* FindChannel(FChannelId ChannelId)
 
 
 ////////////////////////////////////////////////////////////////////////////////
-UE_TRACE_CHANNEL_EXTERN(TraceLogChannel)
+UE_TRACE_MINIMAL_CHANNEL_EXTERN(TraceLogChannel)
 
-UE_TRACE_EVENT_BEGIN($Trace, ThreadInfo, NoSync|Important)
-	UE_TRACE_EVENT_FIELD(uint32, ThreadId)
-	UE_TRACE_EVENT_FIELD(uint32, SystemId)
-	UE_TRACE_EVENT_FIELD(int32, SortHint)
-	UE_TRACE_EVENT_FIELD(AnsiString, Name)
-UE_TRACE_EVENT_END()
+UE_TRACE_MINIMAL_EVENT_BEGIN($Trace, ThreadInfo, NoSync|Important)
+	UE_TRACE_MINIMAL_EVENT_FIELD(uint32, ThreadId)
+	UE_TRACE_MINIMAL_EVENT_FIELD(uint32, SystemId)
+	UE_TRACE_MINIMAL_EVENT_FIELD(int32, SortHint)
+	UE_TRACE_MINIMAL_EVENT_FIELD(AnsiString, Name)
+UE_TRACE_MINIMAL_EVENT_END()
 
-UE_TRACE_EVENT_BEGIN($Trace, ThreadGroupBegin, NoSync|Important)
-	UE_TRACE_EVENT_FIELD(AnsiString, Name)
-UE_TRACE_EVENT_END()
+UE_TRACE_MINIMAL_EVENT_BEGIN($Trace, ThreadGroupBegin, NoSync|Important)
+	UE_TRACE_MINIMAL_EVENT_FIELD(AnsiString, Name)
+UE_TRACE_MINIMAL_EVENT_END()
 
-UE_TRACE_EVENT_BEGIN($Trace, ThreadGroupEnd, NoSync|Important)
-UE_TRACE_EVENT_END()
+UE_TRACE_MINIMAL_EVENT_BEGIN($Trace, ThreadGroupEnd, NoSync|Important)
+UE_TRACE_MINIMAL_EVENT_END()
 
+	
 ////////////////////////////////////////////////////////////////////////////////
 void ThreadRegister(const TCHAR* Name, uint32 SystemId, int32 SortHint)
 {
@@ -234,7 +235,7 @@ void ThreadRegister(const TCHAR* Name, uint32 SystemId, int32 SortHint)
 
 	uint32 ThreadId = Private::Writer_GetThreadId();
 	uint32 NameLen = ToAnsiCheap(NameA, Name);
-	UE_TRACE_LOG($Trace, ThreadInfo, TraceLogChannel, NameLen * sizeof(ANSICHAR))
+	UE_TRACE_MINIMAL_LOG($Trace, ThreadInfo, TraceLogChannel, NameLen * sizeof(ANSICHAR))
 		<< ThreadInfo.ThreadId(ThreadId)
 		<< ThreadInfo.SystemId(SystemId)
 		<< ThreadInfo.SortHint(SortHint)
@@ -247,14 +248,14 @@ void ThreadGroupBegin(const TCHAR* Name)
 	ANSICHAR NameA[96];
 
 	uint32 NameLen = ToAnsiCheap(NameA, Name);
-	UE_TRACE_LOG($Trace, ThreadGroupBegin, TraceLogChannel, NameLen * sizeof(ANSICHAR))
+	UE_TRACE_MINIMAL_LOG($Trace, ThreadGroupBegin, TraceLogChannel, NameLen * sizeof(ANSICHAR))
 		<< ThreadGroupBegin.Name(Name, NameLen);
 }
 
 ////////////////////////////////////////////////////////////////////////////////
 void ThreadGroupEnd()
 {
-	UE_TRACE_LOG($Trace, ThreadGroupEnd, TraceLogChannel);
+	UE_TRACE_MINIMAL_LOG($Trace, ThreadGroupEnd, TraceLogChannel);
 }
 
 } // namespace Trace
@@ -265,4 +266,4 @@ void ThreadGroupEnd()
 // Workaround for module not having any exported symbols
 TRACELOG_API int TraceLogExportedSymbol = 0;
 
-#endif // UE_TRACE_ENABLED
+#endif // TRACE_PRIVATE_MINIMAL_ENABLED
diff --git a/Engine/Source/Runtime/TraceLog/Private/Trace/Writer.cpp b/Engine/Source/Runtime/TraceLog/Private/Trace/Writer.cpp
--- a/Engine/Source/Runtime/TraceLog/Private/Trace/Writer.cpp
+++ b/Engine/Source/Runtime/TraceLog/Private/Trace/Writer.cpp
@@ -3,7 +3,7 @@
 #include "Message.h"
 #include "Trace/Config.h"
 
-#if UE_TRACE_ENABLED
+#if TRACE_PRIVATE_MINIMAL_ENABLED
 
 #include "Platform.h"
 #include "Trace/Detail/Atomic.h"
@@ -69,14 +69,12 @@ struct FTraceGuid
 };
 
 ////////////////////////////////////////////////////////////////////////////////
-UE_TRACE_EVENT_BEGIN($Trace, NewTrace, Important|NoSync)
-	UE_TRACE_EVENT_FIELD(uint64, StartCycle)
-	UE_TRACE_EVENT_FIELD(uint64, CycleFrequency)
-	UE_TRACE_EVENT_FIELD(uint16, Endian)
-	UE_TRACE_EVENT_FIELD(uint8, PointerSize)
-UE_TRACE_EVENT_END()
-
-
+UE_TRACE_MINIMAL_EVENT_BEGIN($Trace, NewTrace, Important | NoSync)
+	UE_TRACE_MINIMAL_EVENT_FIELD(uint64, StartCycle)
+	UE_TRACE_MINIMAL_EVENT_FIELD(uint64, CycleFrequency)
+	UE_TRACE_MINIMAL_EVENT_FIELD(uint16, Endian)
+	UE_TRACE_MINIMAL_EVENT_FIELD(uint8, PointerSize)
+UE_TRACE_MINIMAL_EVENT_END()
 
 ////////////////////////////////////////////////////////////////////////////////
 static volatile bool			GInitialized;		// = false;
@@ -779,7 +777,7 @@ static void Writer_InternalInitializeImpl()
 
 	AtomicStoreRelaxed(&GInitialized, true);
 
-	UE_TRACE_LOG($Trace, NewTrace, TraceLogChannel)
+	UE_TRACE_MINIMAL_LOG($Trace, NewTrace, TraceLogChannel)
 		<< NewTrace.StartCycle(GStartCycle)
 		<< NewTrace.CycleFrequency(TimeGetFrequency())
 		<< NewTrace.Endian(uint16(0x524d))
@@ -923,6 +921,7 @@ static UPTRINT Writer_PackSendFlags(UPTRINT DataHandle, uint32 Flags)
 ////////////////////////////////////////////////////////////////////////////////
 bool Writer_SendTo(const ANSICHAR* Host, uint32 Flags, uint32 Port)
 {
+#if TRACE_PRIVATE_ALLOW_TCP
 	if (AtomicLoadRelaxed(&GPendingDataHandle))
 	{
 		return false;
@@ -947,11 +946,15 @@ bool Writer_SendTo(const ANSICHAR* Host, uint32 Flags, uint32 Port)
 
 	AtomicStoreRelaxed(&GPendingDataHandle, DataHandle);
 	return true;
+#else
+	return false;
+#endif
 }
 
 ////////////////////////////////////////////////////////////////////////////////
 bool Writer_WriteTo(const ANSICHAR* Path, uint32 Flags)
 {
+#if TRACE_PRIVATE_ALLOW_FILE
 	if (AtomicLoadRelaxed(&GPendingDataHandle))
 	{
 		return false;
@@ -974,6 +977,9 @@ bool Writer_WriteTo(const ANSICHAR* Path, uint32 Flags)
 
 	AtomicStoreRelaxed(&GPendingDataHandle, DataHandle);
 	return true;
+#else
+	return false;
+#endif
 }
 
 ////////////////////////////////////////////////////////////////////////////////
@@ -1107,11 +1113,14 @@ bool Writer_WriteSnapshot(const FSnapshotTarget& Target)
 
 		if (Target.Type == FSnapshotTarget::EType::FileTarget)
 		{
+#if TRACE_PRIVATE_ALLOW_FILE
 			// Open the snapshot file 
 			GDataHandle = FileOpen(Target.File.Path);
+#endif
 		}
 		else
 		{
+#if TRACE_PRIVATE_ALLOW_TCP
 			// Open the snapshot connection and write 
 			const uint32 Port = Target.Host.Port ? Target.Host.Port : 1981;
 			GDataHandle = TcpSocketConnect(Target.Host.Host, uint16(Port));
@@ -1120,6 +1129,7 @@ bool Writer_WriteSnapshot(const FSnapshotTarget& Target)
 				return false;
 			}
 			GDataHandle = Writer_PackSendFlags(GDataHandle, 0);
+#endif
 		}
 
 		// Write the file header
@@ -1208,4 +1218,4 @@ bool Writer_Stop()
 } // namespace Trace
 } // namespace UE
 
-#endif // UE_TRACE_ENABLED
+#endif // TRACE_PRIVATE_MINIMAL_ENABLED
diff --git a/Engine/Source/Runtime/TraceLog/Public/Trace/Config.h b/Engine/Source/Runtime/TraceLog/Public/Trace/Config.h
--- a/Engine/Source/Runtime/TraceLog/Public/Trace/Config.h
+++ b/Engine/Source/Runtime/TraceLog/Public/Trace/Config.h
@@ -1,9 +1,14 @@
 // Copyright Epic Games, Inc. All Rights Reserved.
 
 #pragma once
-
 #include "HAL/Platform.h"
 
+/**
+ * Public control for enabling tracing.
+ *
+ * This define controls all the TraceLog functionality and is
+ * used under normal circumstances.
+ */
 #if !defined(UE_TRACE_ENABLED)
 #	if !UE_BUILD_SHIPPING && !IS_PROGRAM
 #		if PLATFORM_WINDOWS || PLATFORM_UNIX || PLATFORM_APPLE || PLATFORM_ANDROID
@@ -16,12 +21,102 @@
 #	define UE_TRACE_ENABLED 0
 #endif
 
+/**
+ * EXPERIMENTAL!: Optional control to enable tracing in shipping configuration.
+ *
+ * When this define is enabled TraceLog functionality can be enabled in
+ * shipping builds. Note that regular tracing system that relies on UE_TRACE_ENABLED
+ * will not be active. This is intentional in order to avoid unintentional functionality
+ * slipping into shipping builds.
+ */
+#if !defined(UE_TRACE_ENABLED_SHIPPING_EXPERIMENTAL)
+#	define UE_TRACE_ENABLED_SHIPPING_EXPERIMENTAL 0
+#endif
+
+
+/**
+ * Public control for the minimal set of tracing.
+ */
+#if !defined(UE_TRACE_MINIMAL_ENABLED)
+#	define UE_TRACE_MINIMAL_ENABLED UE_TRACE_ENABLED
+#endif
+
+/**
+ * Internal defined used inside this library. Do not set outside of TraceLog.
+ */
+#if UE_TRACE_ENABLED || UE_TRACE_MINIMAL_ENABLED
+#	define TRACE_PRIVATE_MINIMAL_ENABLED 1
+#else
+#	define TRACE_PRIVATE_MINIMAL_ENABLED 0
+#endif
+
 #if UE_TRACE_ENABLED
+#	define TRACE_PRIVATE_FULL_ENABLED 1
+#else
+#	define TRACE_PRIVATE_FULL_ENABLED 0
+#endif
+
+/**
+ * Shipping trace is a subset of regular tracing. If regular trace is
+ * enabled, shipping trace must also be enabled.
+ */
+static_assert(TRACE_PRIVATE_FULL_ENABLED == 0 || TRACE_PRIVATE_MINIMAL_ENABLED != 0, "Minimal is a subset of full tracing.");
+#if UE_BUILD_SHIPPING
+static_assert(!TRACE_PRIVATE_FULL_ENABLED, "Full tracing in shipping is not supported.");
+static_assert(UE_TRACE_MINIMAL_ENABLED == 0 || UE_TRACE_ENABLED_SHIPPING_EXPERIMENTAL, "Minimal trace in shipping is currently experimental");
+#endif
+
+
+/**
+ * Current protocol version
+ */
+#if TRACE_PRIVATE_MINIMAL_ENABLED
 #	define TRACE_PRIVATE_PROTOCOL_7
 #endif
 
-// Enable packet verification. Only useful when looking for transmission bugs. Note that in order to avoid making a new protocol version, enabling
-// this makes existing version 7 traces incompatible.
+/**
+ * Control the socket control component. By default we disable it if only
+ * shipping trace is enabled.
+ */
+#if !defined(UE_TRACE_ALLOW_TCP_CONTROL)
+#	define TRACE_PRIVATE_ALLOW_TCP_CONTROL TRACE_PRIVATE_FULL_ENABLED
+#else
+#	define TRACE_PRIVATE_ALLOW_TCP_CONTROL UE_TRACE_ALLOW_TCP_CONTROL
+#endif
+
+/**
+ * Control tracing to tcp socket connections.
+ */
+#if !defined(UE_TRACE_ALLOW_TCP_TRACING)
+#	define TRACE_PRIVATE_ALLOW_TCP TRACE_PRIVATE_FULL_ENABLED
+#else
+#	define TRACE_PRIVATE_ALLOW_TCP UE_TRACE_ALLOW_TCP_TRACING
+#endif
+
+/**
+ * Control tracing to files
+ */
+#if !defined(UE_TRACE_ALLOW_FILE_TRACING)
+#	define TRACE_PRIVATE_ALLOW_FILE TRACE_PRIVATE_FULL_ENABLED
+#else
+#	define TRACE_PRIVATE_ALLOW_FILE UE_TRACE_ALLOW_FILE_TRACING
+#endif
+
+/**
+ * Control if important events are enabled
+ */
+#if !defined(UE_TRACE_ALLOW_IMPORTANT_EVENTS)
+//@todo Disable by default for shipping trace when we have fixed the built in events
+//#	define TRACE_PRIVATE_ALLOW_IMPORTANTS TRACE_PRIVATE_FULL_ENABLED
+#	define TRACE_PRIVATE_ALLOW_IMPORTANTS 1
+#else
+#	define TRACE_PRIVATE_ALLOW_IMPORTANTS UE_TRACE_ALLOW_IMPORTANT_EVENTS
+#endif
+
+/**
+ * Enable packet verification. Only useful when looking for transmission bugs. Note that in order
+ * to avoid making a new protocol version, enabling this makes existing version 7 traces incompatible.
+ */
 #ifndef UE_TRACE_PACKET_VERIFICATION
-#define UE_TRACE_PACKET_VERIFICATION 0
+#	define UE_TRACE_PACKET_VERIFICATION 0
 #endif
\ No newline at end of file
diff --git a/Engine/Source/Runtime/TraceLog/Public/Trace/Detail/Channel.h b/Engine/Source/Runtime/TraceLog/Public/Trace/Detail/Channel.h
--- a/Engine/Source/Runtime/TraceLog/Public/Trace/Detail/Channel.h
+++ b/Engine/Source/Runtime/TraceLog/Public/Trace/Detail/Channel.h
@@ -5,7 +5,7 @@
 #include "Trace/Config.h"
 #include "Trace/Trace.h"
 
-#if UE_TRACE_ENABLED
+#if TRACE_PRIVATE_MINIMAL_ENABLED
 
 #include "CoreTypes.h"
 
@@ -88,4 +88,11 @@ inline uint32 FChannel::GetName(const ANSICHAR** OutName) const
 } // namespace Trace
 } // namespace UE
 
-#endif // UE_TRACE_ENABLED
+#else
+
+// Since we use this type in macros we need
+// provide an empty definition when trace is
+// not enabled.
+namespace UE::Trace { class FChannel {}; }
+
+#endif // TRACE_PRIVATE_MINIMAL_ENABLED
diff --git a/Engine/Source/Runtime/TraceLog/Public/Trace/Detail/Channel.inl b/Engine/Source/Runtime/TraceLog/Public/Trace/Detail/Channel.inl
--- a/Engine/Source/Runtime/TraceLog/Public/Trace/Detail/Channel.inl
+++ b/Engine/Source/Runtime/TraceLog/Public/Trace/Detail/Channel.inl
@@ -4,7 +4,7 @@
 
 #include "Channel.h"
 
-#if UE_TRACE_ENABLED
+#if TRACE_PRIVATE_MINIMAL_ENABLED
 
 namespace UE {
 namespace Trace {
@@ -32,4 +32,4 @@ inline bool FChannel::operator | (const FChannel& Rhs) const
 } // namespace Trace
 } // namespace UE
 
-#endif // UE_TRACE_ENABLED
+#endif // TRACE_PRIVATE_MINIMAL_ENABLED
diff --git a/Engine/Source/Runtime/TraceLog/Public/Trace/Detail/EventNode.h b/Engine/Source/Runtime/TraceLog/Public/Trace/Detail/EventNode.h
--- a/Engine/Source/Runtime/TraceLog/Public/Trace/Detail/EventNode.h
+++ b/Engine/Source/Runtime/TraceLog/Public/Trace/Detail/EventNode.h
@@ -5,7 +5,7 @@
 #include "HAL/Platform.h"
 #include "Trace/Config.h"
 
-#if UE_TRACE_ENABLED
+#if TRACE_PRIVATE_MINIMAL_ENABLED
 
 #include "Field.h"
 
@@ -64,4 +64,11 @@ private:
 } // namespace Trace
 } // namespace UE
 
-#endif // UE_TRACE_ENABLED
+#else
+
+// Since we use this type in macros we need
+// provide an empty definition when trace is
+// not enabled.
+namespace UE::Trace::Private { struct FEventInfo {}; }
+
+#endif // TRACE_PRIVATE_MINIMAL_ENABLED
diff --git a/Engine/Source/Runtime/TraceLog/Public/Trace/Detail/Field.h b/Engine/Source/Runtime/TraceLog/Public/Trace/Detail/Field.h
--- a/Engine/Source/Runtime/TraceLog/Public/Trace/Detail/Field.h
+++ b/Engine/Source/Runtime/TraceLog/Public/Trace/Detail/Field.h
@@ -4,7 +4,7 @@
 
 #include "Trace/Config.h"
 
-#if UE_TRACE_ENABLED
+#if TRACE_PRIVATE_MINIMAL_ENABLED
 
 #include "Atomic.h"
 #include "Protocol.h"
@@ -216,4 +216,4 @@ struct TField<InNumFields, InSize, EventProps>
 } // namespace Trace
 } // namespace UE
 
-#endif // UE_TRACE_ENABLED
+#endif // TRACE_PRIVATE_MINIMAL_ENABLED
diff --git a/Engine/Source/Runtime/TraceLog/Public/Trace/Detail/Important/ImportantLogScope.h b/Engine/Source/Runtime/TraceLog/Public/Trace/Detail/Important/ImportantLogScope.h
--- a/Engine/Source/Runtime/TraceLog/Public/Trace/Detail/Important/ImportantLogScope.h
+++ b/Engine/Source/Runtime/TraceLog/Public/Trace/Detail/Important/ImportantLogScope.h
@@ -4,11 +4,11 @@
 
 #include "Trace/Config.h"
 
-#if UE_TRACE_ENABLED
-
 namespace UE {
 namespace Trace {
 namespace Private {
+	
+#if TRACE_PRIVATE_MINIMAL_ENABLED && TRACE_PRIVATE_ALLOW_IMPORTANTS
 
 ////////////////////////////////////////////////////////////////////////////////
 class FImportantLogScope
@@ -32,8 +32,26 @@ private:
 	int32						AuxCursor;
 };
 
+#else
+
+class FImportantLogScope
+{
+public:
+	template <typename EventType>
+	static FImportantLogScope	Enter() { return FImportantLogScope(); }
+	template <typename EventType>
+	static FImportantLogScope	Enter(uint32 ArrayDataSize) { return FImportantLogScope(); }
+	void						operator += (const FImportantLogScope&) const;
+	const FImportantLogScope&	operator << (bool) const	{ return *this; }
+	constexpr explicit			operator bool () const		{ return true; }
+
+	template <typename FieldMeta, typename Type>
+	struct FFieldSet;
+};
+
+#endif // TRACE_PRIVATE_MINIMAL_ENABLED && TRACE_PRIVATE_ALLOW_IMPORTANTS
+	
 } // namespace Private
 } // namespace Trace
 } // namespace UE
 
-#endif // UE_TRACE_ENABLED
diff --git a/Engine/Source/Runtime/TraceLog/Public/Trace/Detail/Important/ImportantLogScope.inl b/Engine/Source/Runtime/TraceLog/Public/Trace/Detail/Important/ImportantLogScope.inl
--- a/Engine/Source/Runtime/TraceLog/Public/Trace/Detail/Important/ImportantLogScope.inl
+++ b/Engine/Source/Runtime/TraceLog/Public/Trace/Detail/Important/ImportantLogScope.inl
@@ -3,9 +3,6 @@
 #pragma once
 
 #include "Trace/Config.h"
-
-#if UE_TRACE_ENABLED
-
 #include "ImportantLogScope.h"
 #include "SharedBuffer.h"
 #include "Trace/Detail/Protocol.h"
@@ -16,6 +13,8 @@
 namespace UE {
 namespace Trace {
 namespace Private {
+	
+#if TRACE_PRIVATE_MINIMAL_ENABLED && TRACE_PRIVATE_ALLOW_IMPORTANTS
 
 ////////////////////////////////////////////////////////////////////////////////
 extern TRACELOG_API FSharedBuffer* volatile GSharedBuffer;
@@ -217,8 +216,58 @@ struct FImportantLogScope::FFieldSet<FieldMeta, TEventRef<DefinitionType>>
 	}
 };
 
+#else // TRACE_PRIVATE_MINIMAL_ENABLED && TRACE_PRIVATE_ALLOW_IMPORTANTS
+      
+
+template <typename FieldMeta, typename Type>
+struct FImportantLogScope::FFieldSet
+{
+	static void Impl(FImportantLogScope* Scope, const Type& Value)
+	{
+	}
+};
+	
+template <typename FieldMeta, typename Type>
+struct FImportantLogScope::FFieldSet<FieldMeta, Type[]>
+{
+	static void Impl(FImportantLogScope* Scope, Type const* Data, int32 Num)
+	{
+	}
+};
+
+template <typename FieldMeta>
+struct FImportantLogScope::FFieldSet<FieldMeta, AnsiString>
+{
+	static void Impl(FImportantLogScope* Scope, const ANSICHAR* String, int32 Length=-1)
+	{
+	}
+
+	static void Impl(FImportantLogScope* Scope, const WIDECHAR* String, int32 Length=-1)
+	{
+	}
+};
+
+template <typename FieldMeta>
+struct FImportantLogScope::FFieldSet<FieldMeta, WideString>
+{
+	static void Impl(FImportantLogScope* Scope, const WIDECHAR* String, int32 Length=-1)
+	{
+	}
+};
+
+template <typename FieldMeta, typename DefinitionType>
+struct FImportantLogScope::FFieldSet<FieldMeta, TEventRef<DefinitionType>>
+{
+	static void Impl(FImportantLogScope* Scope, const TEventRef<DefinitionType>& Reference)
+	{
+	}
+};
+	
+
+#endif // TRACE_PRIVATE_MINIMAL_ENABLED && TRACE_PRIVATE_ALLOW_IMPORTANTS
+
 } // namespace Private
 } // namespace Trace
 } // namespace UE
 
-#endif // UE_TRACE_ENABLED
+
diff --git a/Engine/Source/Runtime/TraceLog/Public/Trace/Detail/Important/SharedBuffer.h b/Engine/Source/Runtime/TraceLog/Public/Trace/Detail/Important/SharedBuffer.h
--- a/Engine/Source/Runtime/TraceLog/Public/Trace/Detail/Important/SharedBuffer.h
+++ b/Engine/Source/Runtime/TraceLog/Public/Trace/Detail/Important/SharedBuffer.h
@@ -5,7 +5,7 @@
 #include "HAL/Platform.h"
 #include "Trace/Config.h"
 
-#if UE_TRACE_ENABLED
+#if TRACE_PRIVATE_MINIMAL_ENABLED && TRACE_PRIVATE_ALLOW_IMPORTANTS
 
 namespace UE {
 namespace Trace {
@@ -37,4 +37,4 @@ struct FNextSharedBuffer
 } // namespace Trace
 } // namespace UE
 
-#endif // UE_TRACE_ENABLED
+#endif // TRACE_PRIVATE_MINIMAL_ENABLED && TRACE_PRIVATE_ALLOW_IMPORTANTS
diff --git a/Engine/Source/Runtime/TraceLog/Public/Trace/Detail/LogScope.h b/Engine/Source/Runtime/TraceLog/Public/Trace/Detail/LogScope.h
--- a/Engine/Source/Runtime/TraceLog/Public/Trace/Detail/LogScope.h
+++ b/Engine/Source/Runtime/TraceLog/Public/Trace/Detail/LogScope.h
@@ -4,7 +4,7 @@
 
 #include "Trace/Config.h"
 
-#if UE_TRACE_ENABLED
+#if TRACE_PRIVATE_MINIMAL_ENABLED
 
 #include "Writer.inl"
 
@@ -80,4 +80,4 @@ public:
 } // namespace Trace
 } // namespace UE
 
-#endif // UE_TRACE_ENABLED
+#endif // TRACE_PRIVATE_MINIMAL_ENABLED
diff --git a/Engine/Source/Runtime/TraceLog/Public/Trace/Detail/LogScope.inl b/Engine/Source/Runtime/TraceLog/Public/Trace/Detail/LogScope.inl
--- a/Engine/Source/Runtime/TraceLog/Public/Trace/Detail/LogScope.inl
+++ b/Engine/Source/Runtime/TraceLog/Public/Trace/Detail/LogScope.inl
@@ -4,7 +4,7 @@
 
 #include "Trace/Config.h"
 
-#if UE_TRACE_ENABLED
+#if TRACE_PRIVATE_MINIMAL_ENABLED
 
 #include "Atomic.h"
 #include "EventNode.h"
@@ -334,4 +334,4 @@ struct FLogScope::FFieldSet<FieldMeta, TEventRef<DefinitionType>>
 } // namespace Trace
 } // namespace UE
 
-#endif // UE_TRACE_ENABLED
+#endif // TRACE_PRIVATE_MINIMAL_ENABLED
diff --git a/Engine/Source/Runtime/TraceLog/Public/Trace/Detail/Trace.h b/Engine/Source/Runtime/TraceLog/Public/Trace/Detail/Trace.h
--- a/Engine/Source/Runtime/TraceLog/Public/Trace/Detail/Trace.h
+++ b/Engine/Source/Runtime/TraceLog/Public/Trace/Detail/Trace.h
@@ -4,7 +4,6 @@
 
 #include "Trace/Config.h"
 
-#if UE_TRACE_ENABLED
 
 #include "Misc/Launder.h"
 
@@ -18,6 +17,14 @@ class FChannel;
 } // namespace Trace
 } // namespace UE
 
+#define TRACE_PRIVATE_EXPAND(x) x
+
+#if TRACE_PRIVATE_MINIMAL_ENABLED
+
+//------------------------------------------------------------------------------------------------
+//-- Private functional variants
+//------------------------------------------------------------------------------------------------
+
 #define TRACE_PRIVATE_STATISTICS (!(UE_BUILD_SHIPPING || USING_THREAD_SANITISER))
 
 #define TRACE_PRIVATE_CHANNEL_DEFAULT_ARGS false, "None"
@@ -139,6 +146,7 @@ class FChannel;
 		} \
 		typedef UE::Trace::TField<DefinitionId_Meta::Index + NumDefinitionFields, DefinitionId_Meta::Offset + DefinitionId_Meta::Size, UE::Trace::EventProps> EventProps_Meta; \
 		EventProps_Meta const EventProps_Private = {}; \
+		static_assert(!bIsImportant || TRACE_PRIVATE_ALLOW_IMPORTANTS, "Important events are disabled in this configuration.");\
 		typedef std::conditional<bIsImportant != 0, UE::Trace::Private::FImportantLogScope, UE::Trace::Private::FLogScope>::type LogScopeType; \
 		explicit operator bool () const { return true; } \
 		enum { EventFlags = PartialEventFlags|((EventProps_Meta::NumAuxFields != 0) ? UE::Trace::Private::FEventInfo::Flag_MaybeHasAux : 0), }; \
@@ -197,30 +205,44 @@ class FChannel;
 	TRACE_PRIVATE_LOG(LoggerName, EventName, ChannelsExpr, ##__VA_ARGS__) \
 		<< EventName.SetDefinitionId(UE::Trace::MakeEventRef(Id, F##LoggerName##EventName##Fields::GetUid()))
 
-#else
+#endif // TRACE_PRIVATE_MINIMAL_ENABLED
 
-#define TRACE_PRIVATE_CHANNEL(ChannelName, ...)
+//------------------------------------------------------------------------------------------------
+//-- Private no op variants
+//------------------------------------------------------------------------------------------------
 
-#define TRACE_PRIVATE_CHANNEL_EXTERN(ChannelName, ...)
+#define TRACE_PRIVATE_NOP_CHANNEL(ChannelName, ...)
 
-#define TRACE_PRIVATE_CHANNEL_DEFINE(ChannelName, ...)
+#define TRACE_PRIVATE_NOP_CHANNEL_EXTERN(ChannelName, ...)
 
-#define TRACE_PRIVATE_CHANNELEXPR_IS_ENABLED(ChannelsExpr) \
+#define TRACE_PRIVATE_NOP_CHANNEL_DEFINE(ChannelName, ...)
+
+#define TRACE_PRIVATE_NOP_CHANNELEXPR_IS_ENABLED(ChannelsExpr) \
 	false
 
-#define TRACE_PRIVATE_EVENT_DEFINE(LoggerName, EventName) \
+#define TRACE_PRIVATE_NOP_EVENT_DEFINE(LoggerName, EventName) \
 	int8* LoggerName##EventName##DummyPtr = nullptr;
 
-#define TRACE_PRIVATE_EVENT_BEGIN(LoggerName, EventName, ...) \
-	TRACE_PRIVATE_EVENT_BEGIN_IMPL(LoggerName, EventName)
+#define TRACE_PRIVATE_NOP_EVENT_BEGIN(LoggerName, EventName, ...) \
+	TRACE_PRIVATE_NOP_EVENT_BEGIN_IMPL(LoggerName, EventName)
 
-#define TRACE_PRIVATE_EVENT_BEGIN_EXTERN(LoggerName, EventName, ...) \
+#define TRACE_PRIVATE_NOP_EVENT_BEGIN_EXTERN(LoggerName, EventName, ...) \
 	extern int8* LoggerName##EventName##DummyPtr; \
-	TRACE_PRIVATE_EVENT_BEGIN_IMPL(LoggerName, EventName)
+	TRACE_PRIVATE_NOP_EVENT_BEGIN_IMPL(LoggerName, EventName)
 
-#define TRACE_PRIVATE_EVENT_BEGIN_IMPL(LoggerName, EventName) \
+#define TRACE_PRIVATE_NOP_EVENT_BEGIN_IMPL(LoggerName, EventName) \
 	struct F##LoggerName##EventName##Dummy \
 	{ \
+		enum \
+		{ \
+			Important, \
+			NoSync, \
+			Definition8bit, \
+			Definition16bit, \
+			Definition32bit, \
+			Definition64bit, \
+			DefinitionBits, \
+		}; \
 		struct FTraceDisabled \
 		{ \
 			const FTraceDisabled& operator () (...) const { return *this; } \
@@ -231,40 +253,38 @@ class FChannel;
 		} \
 		explicit operator bool () const { return false; }
 
-#define TRACE_PRIVATE_EVENT_FIELD(FieldType, FieldName) \
+#define TRACE_PRIVATE_NOP_EVENT_FIELD(FieldType, FieldName) \
 		const FTraceDisabled& FieldName;
 
-#define TRACE_PRIVATE_EVENT_REFFIELD(RefLogger, RefEventType, FieldName) \
+#define TRACE_PRIVATE_NOP_EVENT_REFFIELD(RefLogger, RefEventType, FieldName) \
 		const FTraceDisabled& FieldName;
 
-#define TRACE_PRIVATE_EVENT_END() \
+#define TRACE_PRIVATE_NOP_EVENT_END() \
 	};
 
-#define TRACE_PRIVATE_LOG(LoggerName, EventName, ...) \
+#define TRACE_PRIVATE_NOP_LOG(LoggerName, EventName, ...) \
 	if (const auto& EventName = *(F##LoggerName##EventName##Dummy*)1) \
 		EventName
 
-#define TRACE_PRIVATE_LOG_SCOPED(LoggerName, EventName, ...) \
+#define TRACE_PRIVATE_NOP_LOG_SCOPED(LoggerName, EventName, ...) \
 	if (const auto& EventName = *(F##LoggerName##EventName##Dummy*)1) \
 		EventName
 
-#define TRACE_PRIVATE_LOG_SCOPED_CONDITIONAL(LoggerName, EventName, Condition, ...) \
+#define TRACE_PRIVATE_NOP_LOG_SCOPED_CONDITIONAL(LoggerName, EventName, ...) \
 	if (const auto& EventName = *(F##LoggerName##EventName##Dummy*)1) \
 		EventName
 
-#define TRACE_PRIVATE_LOG_SCOPED_T(LoggerName, EventName, ...) \
+#define TRACE_PRIVATE_NOP_LOG_SCOPED_T(LoggerName, EventName, ...) \
 	if (const auto& EventName = *(F##LoggerName##EventName##Dummy*)1) \
 		EventName
 
-#define TRACE_PRIVATE_LOG_SCOPED_T_CONDITIONAL(LoggerName, EventName, Condition, ...) \
+#define TRACE_PRIVATE_NOP_LOG_SCOPED_T_CONDITIONAL(LoggerName, EventName, ...) \
 	if (const auto& EventName = *(F##LoggerName##EventName##Dummy*)1) \
 		EventName
 
-#define TRACE_PRIVATE_GET_DEFINITION_TYPE_ID(LoggerName, EventName) \
+#define TRACE_PRIVATE_NOP_GET_DEFINITION_TYPE_ID(LoggerName, EventName) \
 	0
 
-#define TRACE_PRIVATE_LOG_DEFINITION(LoggerName, EventName, Id, ChannelsExpr, ...) \
+#define TRACE_PRIVATE_NOP_LOG_DEFINITION(LoggerName, EventName, Id, ChannelsExpr, ...) \
 	UE::Trace::MakeEventRef(Id, 0); \
-	TRACE_PRIVATE_LOG(LoggerName, EventName, ChannelsExpr, ##__VA_ARGS__)
-
-#endif // UE_TRACE_ENABLED
+	TRACE_PRIVATE_NOP_LOG(LoggerName, EventName, ChannelsExpr, ##__VA_ARGS__)
diff --git a/Engine/Source/Runtime/TraceLog/Public/Trace/Detail/Writer.inl b/Engine/Source/Runtime/TraceLog/Public/Trace/Detail/Writer.inl
--- a/Engine/Source/Runtime/TraceLog/Public/Trace/Detail/Writer.inl
+++ b/Engine/Source/Runtime/TraceLog/Public/Trace/Detail/Writer.inl
@@ -2,7 +2,7 @@
 
 #pragma once
 
-#if UE_TRACE_ENABLED
+#if TRACE_PRIVATE_MINIMAL_ENABLED
 
 namespace UE {
 namespace Trace {
@@ -49,4 +49,4 @@ inline uint64 Writer_GetTimestamp(FWriteBuffer* Buffer)
 } // namespace Trace
 } // namespace UE
 
-#endif // UE_TRACE_ENABLED
+#endif // TRACE_PRIVATE_MINIMAL_ENABLED
diff --git a/Engine/Source/Runtime/TraceLog/Public/Trace/Trace.h b/Engine/Source/Runtime/TraceLog/Public/Trace/Trace.h
--- a/Engine/Source/Runtime/TraceLog/Public/Trace/Trace.h
+++ b/Engine/Source/Runtime/TraceLog/Public/Trace/Trace.h
@@ -5,7 +5,7 @@
 #include "Detail/Trace.h"
 
 ////////////////////////////////////////////////////////////////////////////////
-#if UE_TRACE_ENABLED
+#if TRACE_PRIVATE_MINIMAL_ENABLED
 #	define UE_TRACE_IMPL(...)
 #	define UE_TRACE_API			TRACELOG_API
 #else
@@ -13,6 +13,37 @@
 #	define UE_TRACE_API			inline
 #endif
 
+// msvc seems to have a strange behaviour when it comes to expanding macros.
+#if defined(_MSC_VER)
+
+#if TRACE_PRIVATE_FULL_ENABLED
+#	define TRACE_IMPL(Macro, ...)		TRACE_PRIVATE_EXPAND(TRACE_PRIVATE_##Macro(__VA_ARGS__))
+#else
+#	define TRACE_IMPL(Macro, ...)		TRACE_PRIVATE_EXPAND(TRACE_PRIVATE_NOP_##Macro(__VA_ARGS__))
+#endif
+
+#if TRACE_PRIVATE_MINIMAL_ENABLED
+#	define TRACE_IMPL_MINIMAL(Macro, ...)	TRACE_PRIVATE_EXPAND(TRACE_PRIVATE_##Macro(__VA_ARGS__))
+#else
+#	define TRACE_IMPL_MINIMAL(Macro, ...)	TRACE_PRIVATE_EXPAND(TRACE_PRIVATE_NOP_##Macro(__VA_ARGS__))
+#endif
+
+#else  // defined(_MSC_VER)
+
+#if TRACE_PRIVATE_FULL_ENABLED
+#	define TRACE_IMPL(Macro, ...)		TRACE_PRIVATE_EXPAND(TRACE_PRIVATE_##Macro)(__VA_ARGS__)
+#else
+#	define TRACE_IMPL(Macro, ...)		TRACE_PRIVATE_EXPAND(TRACE_PRIVATE_NOP_##Macro)(__VA_ARGS__)
+#endif
+
+#if TRACE_PRIVATE_MINIMAL_ENABLED
+#	define TRACE_IMPL_MINIMAL(Macro, ...)	TRACE_PRIVATE_EXPAND(TRACE_PRIVATE_##Macro)(__VA_ARGS__)
+#else
+#	define TRACE_IMPL_MINIMAL(Macro, ...)	TRACE_PRIVATE_EXPAND(TRACE_PRIVATE_NOP_##Macro)(__VA_ARGS__)
+#endif
+
+#endif // defined(_MSC_VER)
+
 ////////////////////////////////////////////////////////////////////////////////
 namespace UE {
 namespace Trace {
@@ -193,23 +224,49 @@ UE_TRACE_API FChannel* FindChannel(FChannelId ChannelId) UE_TRACE_IMPL(nullptr);
 } // namespace Trace
 } // namespace UE
 
+
+////////////////////////////////////////////////////////////////////////////////
+/// Tracing macros
+/// Use these to define event types, channel and emit events.
 ////////////////////////////////////////////////////////////////////////////////
-#define UE_TRACE_EVENT_DEFINE(LoggerName, EventName)											TRACE_PRIVATE_EVENT_DEFINE(LoggerName, EventName)
-#define UE_TRACE_EVENT_BEGIN(LoggerName, EventName, ...)										TRACE_PRIVATE_EVENT_BEGIN(LoggerName, EventName, ##__VA_ARGS__)
-#define UE_TRACE_EVENT_BEGIN_EXTERN(LoggerName, EventName, ...)									TRACE_PRIVATE_EVENT_BEGIN_EXTERN(LoggerName, EventName, ##__VA_ARGS__)
-#define UE_TRACE_EVENT_FIELD(FieldType, FieldName)												TRACE_PRIVATE_EVENT_FIELD(FieldType, FieldName)
-#define UE_TRACE_EVENT_REFERENCE_FIELD(RefLogger, RefEvent, FieldName)							TRACE_PRIVATE_EVENT_REFFIELD(RefLogger, RefEvent, FieldName)
-#define UE_TRACE_EVENT_END()																	TRACE_PRIVATE_EVENT_END()
-#define UE_TRACE_LOG(LoggerName, EventName, ChannelsExpr, ...)									TRACE_PRIVATE_LOG(LoggerName, EventName, ChannelsExpr, ##__VA_ARGS__)
-#define UE_TRACE_LOG_SCOPED(LoggerName, EventName, ChannelsExpr, ...)							TRACE_PRIVATE_LOG_SCOPED(LoggerName, EventName, ChannelsExpr, ##__VA_ARGS__)
-#define UE_TRACE_LOG_SCOPED_CONDITIONAL(LoggerName, EventName, ChannelsExpr, Condition, ...)	TRACE_PRIVATE_LOG_SCOPED_CONDITIONAL(LoggerName, EventName, ChannelsExpr, Condition, ##__VA_ARGS__)
-#define UE_TRACE_LOG_SCOPED_T(LoggerName, EventName, ChannelsExpr, ...)							TRACE_PRIVATE_LOG_SCOPED_T(LoggerName, EventName, ChannelsExpr, ##__VA_ARGS__)
-#define UE_TRACE_LOG_SCOPED_T_CONDITIONAL(LoggerName, EventName, ChannelsExpr, Condition, ...)	TRACE_PRIVATE_LOG_SCOPED_T_CONDITIONAL(LoggerName, EventName, ChannelsExpr, Condition, ##__VA_ARGS__)
-#define UE_TRACE_GET_DEFINITION_TYPE_ID(LoggerName, EventName)									TRACE_PRIVATE_GET_DEFINITION_TYPE_ID(LoggerName, EventName)
-#define UE_TRACE_LOG_DEFINITION(LoggerName, EventName, Id, ChannelsExpr, ...)					TRACE_PRIVATE_LOG_DEFINITION(LoggerName, EventName, Id, ChannelsExpr, ##__VA_ARGS__)
+#define UE_TRACE_EVENT_DEFINE(LoggerName, EventName)											TRACE_IMPL(EVENT_DEFINE, LoggerName, EventName)
+#define UE_TRACE_EVENT_BEGIN(LoggerName, EventName, ...)										TRACE_IMPL(EVENT_BEGIN, LoggerName, EventName, ##__VA_ARGS__)
+#define UE_TRACE_EVENT_BEGIN_EXTERN(LoggerName, EventName, ...)									TRACE_IMPL(EVENT_BEGIN_EXTERN, LoggerName, EventName, ##__VA_ARGS__)
+#define UE_TRACE_EVENT_FIELD(FieldType, FieldName)												TRACE_IMPL(EVENT_FIELD, FieldType, FieldName)
+#define UE_TRACE_EVENT_REFERENCE_FIELD(RefLogger, RefEvent, FieldName)							TRACE_IMPL(EVENT_REFFIELD, RefLogger, RefEvent, FieldName)
+#define UE_TRACE_EVENT_END()																	TRACE_IMPL(EVENT_END)
+#define UE_TRACE_LOG(LoggerName, EventName, ChannelsExpr, ...)									TRACE_IMPL(LOG, LoggerName, EventName, ChannelsExpr, ##__VA_ARGS__)
+#define UE_TRACE_LOG_SCOPED(LoggerName, EventName, ChannelsExpr, ...)							TRACE_IMPL(LOG_SCOPED, LoggerName, EventName, ChannelsExpr, ##__VA_ARGS__)
+#define UE_TRACE_LOG_SCOPED_CONDITIONAL(LoggerName, EventName, ChannelsExpr, Condition, ...)	TRACE_IMPL(LOG_SCOPED_CONDITIONAL, LoggerName, EventName, ChannelsExpr, Condition, ##__VA_ARGS__)
+#define UE_TRACE_LOG_SCOPED_T(LoggerName, EventName, ChannelsExpr, ...)							TRACE_IMPL(LOG_SCOPED_T, LoggerName, EventName, ChannelsExpr, ##__VA_ARGS__)
+#define UE_TRACE_LOG_SCOPED_T_CONDITIONAL(LoggerName, EventName, ChannelsExpr, Condition, ...)	TRACE_IMPL(LOG_SCOPED_T_CONDITIONAL, LoggerName, EventName, ChannelsExpr, Condition, ##__VA_ARGS__)
+#define UE_TRACE_GET_DEFINITION_TYPE_ID(LoggerName, EventName)									TRACE_IMPL(GET_DEFINITION_TYPE_ID, LoggerName, EventName)
+#define UE_TRACE_LOG_DEFINITION(LoggerName, EventName, Id, ChannelsExpr, ...)					TRACE_IMPL(LOG_DEFINITION, LoggerName, EventName, Id, ChannelsExpr, ##__VA_ARGS__)
+#define UE_TRACE_CHANNEL(ChannelName, ...)														TRACE_IMPL(CHANNEL, ChannelName, ##__VA_ARGS__)
+#define UE_TRACE_CHANNEL_EXTERN(ChannelName, ...)												TRACE_IMPL(CHANNEL_EXTERN, ChannelName, ##__VA_ARGS__)
+#define UE_TRACE_CHANNEL_DEFINE(ChannelName, ...)												TRACE_IMPL(CHANNEL_DEFINE, ChannelName, ##__VA_ARGS__)
+#define UE_TRACE_CHANNELEXPR_IS_ENABLED(ChannelsExpr)											TRACE_IMPL(CHANNELEXPR_IS_ENABLED, ChannelsExpr)
 
+
+////////////////////////////////////////////////////////////////////////////////
+/// Shipping variants of the macros.
+/// With these macros users can provide a subset of events that are available
+/// both in development and in shipping configurations (provided UE_TRACE_MINIMAL_ENABLED is set).
 ////////////////////////////////////////////////////////////////////////////////
-#define UE_TRACE_CHANNEL(ChannelName, ...)				TRACE_PRIVATE_CHANNEL(ChannelName, ##__VA_ARGS__)
-#define UE_TRACE_CHANNEL_EXTERN(ChannelName, ...)		TRACE_PRIVATE_CHANNEL_EXTERN(ChannelName, ##__VA_ARGS__)
-#define UE_TRACE_CHANNEL_DEFINE(ChannelName, ...)		TRACE_PRIVATE_CHANNEL_DEFINE(ChannelName, ##__VA_ARGS__)
-#define UE_TRACE_CHANNELEXPR_IS_ENABLED(ChannelsExpr)	TRACE_PRIVATE_CHANNELEXPR_IS_ENABLED(ChannelsExpr)
+#define UE_TRACE_MINIMAL_EVENT_DEFINE(LoggerName, EventName)											TRACE_IMPL_MINIMAL(EVENT_DEFINE, LoggerName, EventName)
+#define UE_TRACE_MINIMAL_EVENT_BEGIN(LoggerName, EventName, ...)										TRACE_IMPL_MINIMAL(EVENT_BEGIN, LoggerName, EventName, ##__VA_ARGS__)
+#define UE_TRACE_MINIMAL_EVENT_BEGIN_EXTERN(LoggerName, EventName, ...)									TRACE_IMPL_MINIMAL(EVENT_BEGIN_EXTERN, LoggerName, EventName, ##__VA_ARGS__)
+#define UE_TRACE_MINIMAL_EVENT_FIELD(FieldType, FieldName)												TRACE_IMPL_MINIMAL(EVENT_FIELD, FieldType, FieldName)
+#define UE_TRACE_MINIMAL_EVENT_REFERENCE_FIELD(RefLogger, RefEvent, FieldName)							TRACE_IMPL_MINIMAL(EVENT_REFFIELD, RefLogger, RefEvent, FieldName)
+#define UE_TRACE_MINIMAL_EVENT_END()																	TRACE_IMPL_MINIMAL(EVENT_END)
+#define UE_TRACE_MINIMAL_LOG(LoggerName, EventName, ChannelsExpr, ...)									TRACE_IMPL_MINIMAL(LOG, LoggerName, EventName, ChannelsExpr, ##__VA_ARGS__)
+#define UE_TRACE_MINIMAL_LOG_SCOPED(LoggerName, EventName, ChannelsExpr, ...)							TRACE_IMPL_MINIMAL(LOG_SCOPED, LoggerName, EventName, ChannelsExpr, ##__VA_ARGS__)
+#define UE_TRACE_MINIMAL_LOG_SCOPED_CONDITIONAL(LoggerName, EventName, ChannelsExpr, Condition, ...)	TRACE_IMPL(LOG_SCOPED_CONDITIONAL, LoggerName, EventName, ChannelsExpr, Condition, ##__VA_ARGS__)
+#define UE_TRACE_MINIMAL_LOG_SCOPED_T(LoggerName, EventName, ChannelsExpr, ...)							TRACE_IMPL_MINIMAL(LOG_SCOPED_T, LoggerName, EventName, ChannelsExpr, ##__VA_ARGS__)
+#define UE_TRACE_MINIMAL_LOG_SCOPED_T_CONDITIONAL(LoggerName, EventName, ChannelsExpr, Condition, ...)	TRACE_IMPL(LOG_SCOPED_T_CONDITIONAL, LoggerName, EventName, ChannelsExpr, Condition, ##__VA_ARGS__)
+#define UE_TRACE_MINIMAL_GET_DEFINITION_TYPE_ID(LoggerName, EventName)									TRACE_IMPL_MINIMAL(GET_DEFINITION_TYPE_ID, LoggerName, EventName)
+#define UE_TRACE_MINIMAL_LOG_DEFINITION(LoggerName, EventName, Id, ChannelsExpr, ...)					TRACE_IMPL_MINIMAL(LOG_DEFINITION, LoggerName, EventName, Id, ChannelsExpr, ##__VA_ARGS__)
+#define UE_TRACE_MINIMAL_CHANNEL(ChannelName, ...)														TRACE_IMPL_MINIMAL(CHANNEL, ChannelName, ##__VA_ARGS__)
+#define UE_TRACE_MINIMAL_CHANNEL_EXTERN(ChannelName, ...)												TRACE_IMPL_MINIMAL(CHANNEL_EXTERN, ChannelName, ##__VA_ARGS__)
+#define UE_TRACE_MINIMAL_CHANNEL_DEFINE(ChannelName, ...)												TRACE_IMPL_MINIMAL(CHANNEL_DEFINE, ChannelName, ##__VA_ARGS__)
+#define UE_TRACE_MINIMAL_CHANNELEXPR_IS_ENABLED(ChannelsExpr)											TRACE_IMPL_MINIMAL(CHANNELEXPR_IS_ENABLED, ChannelsExpr)
