Fix static analysis warnings

diff --git a/Engine/Source/Runtime/TraceLog/Public/Trace/Config.h b/Engine/Source/Runtime/TraceLog/Public/Trace/Config.h
--- a/Engine/Source/Runtime/TraceLog/Public/Trace/Config.h
+++ b/Engine/Source/Runtime/TraceLog/Public/Trace/Config.h
@@ -60,7 +60,7 @@
  * Shipping trace is a subset of regular tracing. If regular trace is
  * enabled, shipping trace must also be enabled.
  */
-static_assert(TRACE_PRIVATE_FULL_ENABLED == 0 || TRACE_PRIVATE_MINIMAL_ENABLED != 0, "Minimal is a subset of full tracing.");
+static_assert(TRACE_PRIVATE_FULL_ENABLED == 0 || TRACE_PRIVATE_MINIMAL_ENABLED != 0, "Minimal is a subset of full tracing."); //-V547
 #if UE_BUILD_SHIPPING
 static_assert(!TRACE_PRIVATE_FULL_ENABLED, "Full tracing in shipping is not supported.");
 static_assert(UE_TRACE_MINIMAL_ENABLED == 0 || UE_TRACE_ENABLED_SHIPPING_EXPERIMENTAL, "Minimal trace in shipping is currently experimental");
diff --git a/Engine/Source/Runtime/TraceLog/Public/Trace/Detail/Trace.h b/Engine/Source/Runtime/TraceLog/Public/Trace/Detail/Trace.h
--- a/Engine/Source/Runtime/TraceLog/Public/Trace/Detail/Trace.h
+++ b/Engine/Source/Runtime/TraceLog/Public/Trace/Detail/Trace.h
@@ -146,7 +146,7 @@ class FChannel;
 		} \
 		typedef UE::Trace::TField<DefinitionId_Meta::Index + NumDefinitionFields, DefinitionId_Meta::Offset + DefinitionId_Meta::Size, UE::Trace::EventProps> EventProps_Meta; \
 		EventProps_Meta const EventProps_Private = {}; \
-		static_assert(!bIsImportant || TRACE_PRIVATE_ALLOW_IMPORTANTS, "Important events are disabled in this configuration.");\
+		static_assert(bIsImportant == 0 || TRACE_PRIVATE_ALLOW_IMPORTANTS, "Important events are disabled in this configuration.");\
 		typedef std::conditional<bIsImportant != 0, UE::Trace::Private::FImportantLogScope, UE::Trace::Private::FLogScope>::type LogScopeType; \
 		explicit operator bool () const { return true; } \
 		enum { EventFlags = PartialEventFlags|((EventProps_Meta::NumAuxFields != 0) ? UE::Trace::Private::FEventInfo::Flag_MaybeHasAux : 0), }; \
