Fix static assert to allow program to enable trace in shipping

diff --git a/Engine/Source/Runtime/TraceLog/Public/Trace/Config.h b/Engine/Source/Runtime/TraceLog/Public/Trace/Config.h
--- a/Engine/Source/Runtime/TraceLog/Public/Trace/Config.h
+++ b/Engine/Source/Runtime/TraceLog/Public/Trace/Config.h
@@ -61,7 +61,7 @@
  * enabled, shipping trace must also be enabled.
  */
 static_assert(TRACE_PRIVATE_FULL_ENABLED == 0 || TRACE_PRIVATE_MINIMAL_ENABLED != 0, "Minimal is a subset of full tracing."); //-V547
-#if UE_BUILD_SHIPPING
+#if UE_BUILD_SHIPPING && !IS_PROGRAM
 static_assert(!TRACE_PRIVATE_FULL_ENABLED, "Full tracing in shipping is not supported.");
 static_assert(UE_TRACE_MINIMAL_ENABLED == 0 || UE_TRACE_ENABLED_SHIPPING_EXPERIMENTAL, "Minimal trace in shipping is currently experimental");
 #endif
