Read UE::Trace::FChannel::Enabled atomically to correct TSAN race report

#rnx

diff --git a/Engine/Source/Runtime/TraceLog/Public/Trace/Detail/Channel.inl b/Engine/Source/Runtime/TraceLog/Public/Trace/Detail/Channel.inl
--- a/Engine/Source/Runtime/TraceLog/Public/Trace/Detail/Channel.inl
+++ b/Engine/Source/Runtime/TraceLog/Public/Trace/Detail/Channel.inl
@@ -3,6 +3,7 @@
 #pragma once
 
 #include "Channel.h"
+#include "Trace/Detail/Atomic.h"
 
 #if TRACE_PRIVATE_MINIMAL_ENABLED
 
@@ -14,7 +15,7 @@ extern TRACELOG_API FChannel& TraceLogChannel;
 ////////////////////////////////////////////////////////////////////////////////
 inline bool FChannel::IsEnabled() const
 {
-	return Enabled >= 0;
+	return Private::AtomicLoadRelaxed((volatile int32*) &Enabled) >= 0;
 }
 
 ////////////////////////////////////////////////////////////////////////////////
