Extended error message when snapshot fails during ongoing trace.

On certain platforms and hardware combinations we have seen that creating a snapshot while an ongoing trace fails for unknown reasons. Display a log message explaining this possibility. Unfortunately the message cannot be contain details about platform.

diff --git a/Engine/Source/Runtime/TraceLog/Private/Trace/Writer.cpp b/Engine/Source/Runtime/TraceLog/Private/Trace/Writer.cpp
--- a/Engine/Source/Runtime/TraceLog/Private/Trace/Writer.cpp
+++ b/Engine/Source/Runtime/TraceLog/Private/Trace/Writer.cpp
@@ -1106,6 +1106,8 @@ bool Writer_WriteSnapshot(const FSnapshotTarget& Target)
 	Writer_FlushSendBuffer();
 
 	{
+		const bool bExistingDataHandle = GDataHandle != 0;
+		
 		TStashGlobal DataHandle(GDataHandle);
 		TStashGlobal PendingDataHandle(GPendingDataHandle);
 		TStashGlobal SyncPacketCountdown(GSyncPacketCountdown, GNumSyncPackets);
@@ -1136,6 +1138,11 @@ bool Writer_WriteSnapshot(const FSnapshotTarget& Target)
 		if (!GDataHandle || !Writer_SessionPrologue())
 		{
 			UE_TRACE_ERRORMESSAGE(FileOpenError, GetLastErrorCode());
+			if (bExistingDataHandle)
+			{
+				UE_TRACE_MESSAGE(Display, "Creating a snapshot during ongoing trace "
+					"is known to fail on some combinations of platforms and hardware.");
+			}
 			return false;
 		}
 
