Add additional profiling for TraceLog

* Add optional profiler hooks in TraceLog, and instrument common worker thread functions.
* Add TraceLog stat for tracking `BytesEmitted`, the total number of bytes of event data emitted, but possibly not yet processed by the worker thread.
* Add additional stats for tracking shared memory and block pool usage.

diff --git a/Engine/Source/Runtime/TraceLog/Private/Trace/BlockPool.cpp b/Engine/Source/Runtime/TraceLog/Private/Trace/BlockPool.cpp
--- a/Engine/Source/Runtime/TraceLog/Private/Trace/BlockPool.cpp
+++ b/Engine/Source/Runtime/TraceLog/Private/Trace/BlockPool.cpp
@@ -16,7 +16,7 @@ namespace Private {
 void*		Writer_MemoryAllocate(SIZE_T, uint32);
 void		Writer_MemoryFree(void*, uint32);
 
-
+extern FStatistics GTraceStatistics;
 
 ////////////////////////////////////////////////////////////////////////////////
 struct FPoolPage
@@ -42,7 +42,6 @@ static const uint32						GPoolInitPageSize	= GPoolBlockSize << 6;
 T_ALIGN static FWriteBuffer* volatile	GPoolFreeList;		// = nullptr;
 T_ALIGN static UPTRINT volatile			GPoolFutex;			// = 0
 T_ALIGN static FPoolPage* volatile		GPoolPageList;		// = nullptr;
-static uint32							GPoolUsage;			// = 0;
 #undef T_ALIGN
 
 ////////////////////////////////////////////////////////////////////////////////
@@ -50,7 +49,10 @@ static FPoolBlockList Writer_AllocateBlockList(uint32 PageSize)
 {
 	// The free list is empty so we have to populate it with some new blocks.
 	uint8* PageBase = (uint8*)Writer_MemoryAllocate(PageSize, PLATFORM_CACHE_LINE_SIZE);
-	GPoolUsage += PageSize;
+	
+#if TRACE_PRIVATE_STATISTICS
+	AtomicAddRelaxed(&GTraceStatistics.BlockPoolAllocated, uint64(PageSize));
+#endif
 
 	uint32 BufferSize = GPoolBlockSize;
 	BufferSize -= sizeof(FWriteBuffer);
@@ -180,8 +182,10 @@ void Writer_ShutdownPool()
 	for (auto* Page = AtomicLoadRelaxed(&GPoolPageList); Page != nullptr;)
 	{
 		FPoolPage* NextPage = Page->NextPage;
-		uint32 PageSize = (NextPage == nullptr) ? GPoolBlockSize : GPoolPageSize;
-		Writer_MemoryFree(Page, PageSize);
+		Writer_MemoryFree(Page, Page->AllocSize);
+#if TRACE_PRIVATE_STATISTICS
+		AtomicSubRelease(&GTraceStatistics.BlockPoolAllocated, uint64(Page->AllocSize));
+#endif
 		Page = NextPage;
 	}
 }
diff --git a/Engine/Source/Runtime/TraceLog/Private/Trace/Codec.cpp b/Engine/Source/Runtime/TraceLog/Private/Trace/Codec.cpp
--- a/Engine/Source/Runtime/TraceLog/Private/Trace/Codec.cpp
+++ b/Engine/Source/Runtime/TraceLog/Private/Trace/Codec.cpp
@@ -1,6 +1,7 @@
 // Copyright Epic Games, Inc. All Rights Reserved.
 
 #include "CoreTypes.h"
+#include "Platform.h"
 
 THIRD_PARTY_INCLUDES_START
 #if defined(_MSC_VER)
@@ -29,6 +30,9 @@ namespace Private {
 ////////////////////////////////////////////////////////////////////////////////
 int32 Encode(const void* Src, int32 SrcSize, void* Dest, int32 DestSize)
 {
+#if TRACE_PRIVATE_MINIMAL_ENABLED
+	FProfilerScope _(__func__);
+#endif
 	return TRACE_PRIVATE_LZ4_NAMESPACE LZ4_compress_fast(
 		(const char*)Src,
 		(char*)Dest,
diff --git a/Engine/Source/Runtime/TraceLog/Private/Trace/Important/Cache.cpp b/Engine/Source/Runtime/TraceLog/Private/Trace/Important/Cache.cpp
--- a/Engine/Source/Runtime/TraceLog/Private/Trace/Important/Cache.cpp
+++ b/Engine/Source/Runtime/TraceLog/Private/Trace/Important/Cache.cpp
@@ -1,6 +1,7 @@
 // Copyright Epic Games, Inc. All Rights Reserved.
 
 #include "Trace/Config.h"
+#include "Trace/Platform.h"
 
 #if TRACE_PRIVATE_MINIMAL_ENABLED && TRACE_PRIVATE_ALLOW_IMPORTANTS
 
@@ -65,6 +66,8 @@ static FCacheBuffer* Writer_CacheCreateBuffer(uint32 Size)
 ////////////////////////////////////////////////////////////////////////////////
 static void Writer_CacheCommit(const FCacheBuffer* Collector)
 {
+	FProfilerScope _(__func__);
+	
 	// Check there's enough size to compress Data into.
 	uint32 InputSize = uint32(Collector->Size - Collector->Remaining);
 	uint32 EncodeMaxSize = GetEncodeMaxSize(InputSize);
@@ -103,6 +106,8 @@ static void Writer_CacheCommit(const FCacheBuffer* Collector)
 ////////////////////////////////////////////////////////////////////////////////
 void Writer_CacheData(uint8* Data, uint32 Size)
 {
+	FProfilerScope _(__func__);
+	
 	Writer_SendData(ETransportTid::Importants, Data, Size);
 
 	if (GCacheCollector == nullptr)
@@ -137,6 +142,8 @@ void Writer_CacheData(uint8* Data, uint32 Size)
 ////////////////////////////////////////////////////////////////////////////////
 void Writer_CacheOnConnect()
 {
+	FProfilerScope _(__func__);
+	
 	if (GCacheCollector == nullptr)
 	{
 		return;
@@ -176,12 +183,20 @@ void Writer_ShutdownCache()
 	for (FCacheBuffer* Buffer = GCacheHeadBuffer; Buffer != nullptr;)
 	{
 		FCacheBuffer* Next = Buffer->Next;
-		Writer_MemoryFree(Buffer, GCacheBufferSize);
+		const uint32 CacheBufferSize = Buffer->Size + sizeof(FCacheBuffer);
+		Writer_MemoryFree(Buffer, CacheBufferSize);
+#if TRACE_PRIVATE_STATISTICS
+		GTraceStatistics.CacheAllocated -= CacheBufferSize;
+#endif
 		Buffer = Next;
 	}
 
 	Writer_MemoryFree(GCacheActiveBuffer, GCacheBufferSize);
 	Writer_MemoryFree(GCacheCollector, GCacheCollectorSize);
+#if TRACE_PRIVATE_STATISTICS
+	GTraceStatistics.CacheAllocated -= GCacheBufferSize;
+	GTraceStatistics.CacheAllocated -= GCacheCollectorSize;
+#endif
 }
 
 } // namespace Private
diff --git a/Engine/Source/Runtime/TraceLog/Private/Trace/Important/SharedBuffer.cpp b/Engine/Source/Runtime/TraceLog/Private/Trace/Important/SharedBuffer.cpp
--- a/Engine/Source/Runtime/TraceLog/Private/Trace/Important/SharedBuffer.cpp
+++ b/Engine/Source/Runtime/TraceLog/Private/Trace/Important/SharedBuffer.cpp
@@ -42,9 +42,13 @@ static FSharedBuffer* Writer_CreateSharedBuffer(uint32 SizeHint=0)
 
 	void* Block = Writer_MemoryAllocate(BlockSize, alignof(FSharedBuffer));
 	auto* Buffer = (FSharedBuffer*)(UPTRINT(Block) + BlockSize) - 1;
+	
+#if TRACE_PRIVATE_STATISTICS
+	AtomicAddRelaxed(&GTraceStatistics.SharedBufferAllocated, BlockSize);
+#endif
 
 	Buffer->Size = uint32(UPTRINT(Buffer) - UPTRINT(Block));
-	Buffer->Size -= sizeof(uint32); // to preceed event data with a small header when sending.
+	Buffer->Size -= sizeof(uint32); // to precede event data with a small header when sending.
 	Buffer->Cursor = (Buffer->Size << FSharedBuffer::CursorShift) | FSharedBuffer::RefInit;
 	Buffer->Next = nullptr;
 	Buffer->Final = 0;
@@ -85,6 +89,9 @@ FNextSharedBuffer Writer_NextSharedBuffer(FSharedBuffer* Buffer, int32 RegionSta
 			}
 		}
 
+#if TRACE_PRIVATE_STATISTICS
+		AtomicAddRelaxed(&GTraceStatistics.BytesEmitted, uint64(-Buffer->Cursor >> FSharedBuffer::CursorShift));
+#endif
 		AtomicAddRelease(&(Buffer->Cursor), int32(FSharedBuffer::RefBit));
 
 		// Try and allocate some space in the next buffer.  The next buffer may have insufficient
@@ -117,8 +124,12 @@ static void Writer_RetireSharedBufferImpl()
 	}
 
 	FSharedBuffer* Temp = GTailBuffer->Next;
-	void* Block = (uint8*)GTailBuffer - GTailBuffer->Size - sizeof(uint32);
-	Writer_MemoryFree(Block, GBlockSize);
+	const uint32 AllocSizeBytes = GTailBuffer->Size + sizeof(uint32) + sizeof(FSharedBuffer);
+	void* Block = (uint8*)GTailBuffer - AllocSizeBytes;
+	Writer_MemoryFree(Block, AllocSizeBytes);
+#if TRACE_PRIVATE_STATISTICS
+	AtomicSubRelease(&GTraceStatistics.SharedBufferAllocated, AllocSizeBytes);
+#endif
 	GTailBuffer = Temp;
 	GTailPreSent = 0;
 }
@@ -142,6 +153,7 @@ static void Writer_RetireSharedBuffer()
 ////////////////////////////////////////////////////////////////////////////////
 void Writer_UpdateSharedBuffers()
 {
+	FProfilerScope _(__func__);
 	int32 ThrottleRestore = -1;
 	ON_SCOPE_EXIT {
 		if (ThrottleRestore >= 0)
diff --git a/Engine/Source/Runtime/TraceLog/Private/Trace/Platform.h b/Engine/Source/Runtime/TraceLog/Private/Trace/Platform.h
--- a/Engine/Source/Runtime/TraceLog/Private/Trace/Platform.h
+++ b/Engine/Source/Runtime/TraceLog/Private/Trace/Platform.h
@@ -3,6 +3,7 @@
 #pragma once
 
 #include "Trace/Config.h"
+#include "Trace/Trace.h"
 
 #if TRACE_PRIVATE_MINIMAL_ENABLED
 
@@ -57,6 +58,20 @@ THROTTLE_INL void	ThreadUnthrottle(int32)	THROTTLE_IMPL()
 #undef THROTTLE_INL
 #undef THROTTLE_IMPL
 
+////////////////////////////////////////////////////////////////////////////////
+struct FProfilerScope
+{
+	static OnScopeBeginFunc* OnScopeBegin;
+	static OnScopeEndFunc* OnScopeEnd;
+
+#if TRACE_PRIVATE_FULL_ENABLED
+	FProfilerScope(const ANSICHAR* Label) { if (OnScopeBegin) { OnScopeBegin(Label); }}
+	~FProfilerScope() { if(OnScopeEnd) { OnScopeEnd(); }}
+#else
+	FProfilerScope(const ANSICHAR* Label) {}
+#endif
+};
+	
 } // namespace Private
 } // namespace Trace
 } // namespace UE
diff --git a/Engine/Source/Runtime/TraceLog/Private/Trace/Tail.cpp b/Engine/Source/Runtime/TraceLog/Private/Trace/Tail.cpp
--- a/Engine/Source/Runtime/TraceLog/Private/Trace/Tail.cpp
+++ b/Engine/Source/Runtime/TraceLog/Private/Trace/Tail.cpp
@@ -1,9 +1,11 @@
 // Copyright Epic Games, Inc. All Rights Reserved.
 
+#include "Platform.h"
 #include "Trace/Config.h"
 
 #if TRACE_PRIVATE_MINIMAL_ENABLED
 
+#include "Trace/Detail/Atomic.h"
 #include "Trace/Detail/Protocol.h"
 #include "Trace/Detail/Transport.h"
 
@@ -26,6 +28,8 @@ void		Writer_MemoryFree(void*, uint32);
 void		Writer_SendData(uint32, uint8* __restrict, uint32);
 void		Writer_SendDataRaw(const void*, uint32);
 
+extern FStatistics GTraceStatistics;	
+
 #if UE_TRACE_PACKET_VERIFICATION	
 extern uint64 GPacketSerial;
 #endif
@@ -71,6 +75,9 @@ static_assert(std::is_trivial<FPacketRing>(), "FPacketRing must be trivial");
 void FPacketRing::Initialize(uint32 InSize)
 {
 	Data = (uint8*)Writer_MemoryAllocate(InSize, 16);
+#if TRACE_PRIVATE_STATISTICS
+	AtomicAddRelaxed(&GTraceStatistics.FixedBufferAllocated, uint32(InSize));
+#endif
 	Size = InSize;
 	Reset();
 }
@@ -79,6 +86,9 @@ void FPacketRing::Initialize(uint32 InSize)
 void FPacketRing::Shutdown()
 {
 	Writer_MemoryFree(Data, Size);
+#if TRACE_PRIVATE_STATISTICS
+	AtomicSubRelaxed(&GTraceStatistics.FixedBufferAllocated, uint32(Size));
+#endif
 	Data = nullptr;
 }
 
@@ -205,6 +215,8 @@ static FPacketRing GPacketRing; // = {};
 ////////////////////////////////////////////////////////////////////////////////
 void Writer_TailAppend(uint32 ThreadId, uint8* __restrict Data, uint32 Size)
 {
+	FProfilerScope _(__func__);
+	
 	// Perhaps tail tracing is disabled?
 	if (!GPacketRing.IsActive())
 	{
@@ -264,6 +276,8 @@ void Writer_TailAppend(uint32 ThreadId, uint8* __restrict Data, uint32 Size)
 ////////////////////////////////////////////////////////////////////////////////
 void Writer_TailOnConnect()
 {
+	FProfilerScope _(__func__);
+	
 	// If there's no tail being maintained then there is nothing to send.
 	if (!GPacketRing.IsActive())
 	{
diff --git a/Engine/Source/Runtime/TraceLog/Private/Trace/TlsBuffer.cpp b/Engine/Source/Runtime/TraceLog/Private/Trace/TlsBuffer.cpp
--- a/Engine/Source/Runtime/TraceLog/Private/Trace/TlsBuffer.cpp
+++ b/Engine/Source/Runtime/TraceLog/Private/Trace/TlsBuffer.cpp
@@ -111,6 +111,13 @@ static FWriteBuffer* Writer_NextBufferInternal(FWriteBuffer* CurrentBuffer)
 	}
 	else
 	{
+	#if TRACE_PRIVATE_STATISTICS
+		// Note that we only track the emitted bytes when a buffer is filled as the overhead
+		// of doing on every event would be too large.
+		const uint64 Used = CurrentBuffer->Cursor - ((uint8*)CurrentBuffer - CurrentBuffer->Size);
+		AtomicAddRelaxed(&GTraceStatistics.BytesEmitted, Used);
+	#endif
+		
 		NextBuffer->ThreadId = CurrentBuffer->ThreadId;
 		NextBuffer->PrevTimestamp = CurrentBuffer->PrevTimestamp;
 		AtomicStoreRelease(&(CurrentBuffer->NextBuffer), NextBuffer);
@@ -154,6 +161,7 @@ static bool Writer_DrainBuffer(uint32 ThreadId, FWriteBuffer* Buffer)
 ////////////////////////////////////////////////////////////////////////////////
 void Writer_DrainBuffers()
 {
+	FProfilerScope _(__func__);
 	struct FRetireList
 	{
 		FWriteBuffer* __restrict Head = nullptr;
diff --git a/Engine/Source/Runtime/TraceLog/Private/Trace/Writer.cpp b/Engine/Source/Runtime/TraceLog/Private/Trace/Writer.cpp
--- a/Engine/Source/Runtime/TraceLog/Private/Trace/Writer.cpp
+++ b/Engine/Source/Runtime/TraceLog/Private/Trace/Writer.cpp
@@ -197,6 +197,11 @@ void Writer_CreateGuid(FTraceGuid* OutGuid)
 	Octets[6] = Version | (~VersionMask & Octets[6]); // Octet 9
 	Octets[8] = Variant | (~VariantMask & Octets[8]); // Octet 7
 }
+	
+////////////////////////////////////////////////////////////////////////////////
+
+OnScopeBeginFunc* FProfilerScope::OnScopeBegin = nullptr;
+OnScopeEndFunc* FProfilerScope::OnScopeEnd = nullptr;
 
 ////////////////////////////////////////////////////////////////////////////////
 void*			(*AllocHook)(SIZE_T, uint32);			// = nullptr
@@ -330,6 +335,8 @@ static bool Writer_FlushSendBuffer() { return true; }
 ////////////////////////////////////////////////////////////////////////////////
 static void Writer_SendDataImpl(const void* Data, uint32 Size)
 {
+	FProfilerScope _(__func__);
+
 #if TRACE_PRIVATE_STATISTICS
 	GTraceStatistics.BytesSent += Size;
 #endif
@@ -389,6 +396,8 @@ void Writer_SendData(uint32 ThreadId, uint8* __restrict Data, uint32 Size)
 	{
 		return;
 	}
+	
+	FProfilerScope _(__func__);
 
 	// Smaller buffers usually aren't redundant enough to benefit from being
 	// compressed. They often end up being larger.
@@ -532,6 +541,8 @@ static bool Writer_UpdateConnection()
 		return false;
 	}
 
+	FProfilerScope _(__func__);
+
 	// Is this a close request? So that we capture some of the events around
 	// the closure we will add some inertia before enacting the close.
 	static int32 CloseInertia = 0;
@@ -610,6 +621,9 @@ static bool Writer_SessionPrologue()
 	if (!GSendBuffer)
 	{
 		GSendBuffer = static_cast<uint8*>(Writer_MemoryAllocate(GSendBufferSize, 16));
+#if TRACE_PRIVATE_STATISTICS
+		AtomicAddRelaxed(&GTraceStatistics.FixedBufferAllocated, uint32(GSendBufferSize));
+#endif
 	}
 	GSendBufferCursor = GSendBuffer;
 #endif
@@ -688,6 +702,7 @@ static volatile uint32	GUpdateInProgress = 1;	// Don't allow updates until initi
 ////////////////////////////////////////////////////////////////////////////////
 static void Writer_WorkerUpdateInternal()
 {
+	FProfilerScope _(__func__);
 	Writer_UpdateControl();
 	Writer_UpdateConnection();
 	Writer_DescribeAnnounce();
@@ -811,6 +826,9 @@ static void Writer_InternalShutdown()
 	if (GSendBuffer)
 	{
 		Writer_MemoryFree(GSendBuffer, GSendBufferSize);
+#if TRACE_PRIVATE_STATISTICS
+		AtomicSubRelaxed(&GTraceStatistics.FixedBufferAllocated, uint32(GSendBufferSize));
+#endif
 		GSendBuffer = nullptr;
 		GSendBufferCursor = nullptr;
 	}
@@ -850,6 +868,12 @@ void Writer_InternalInitialize()
 ////////////////////////////////////////////////////////////////////////////////
 void Writer_Initialize(const FInitializeDesc& Desc)
 {
+	// Store scope callbacks for optional profiling
+	FProfilerScope::OnScopeBegin = Desc.OnScopeBeginFunc;
+	FProfilerScope::OnScopeEnd = Desc.OnScopeEndFunc;
+
+	FProfilerScope _(__func__);
+	
 	Writer_InitializeTail(Desc.TailSizeBytes);
 
 	if (Desc.bUseImportantCache)
diff --git a/Engine/Source/Runtime/TraceLog/Public/Trace/Detail/Atomic.h b/Engine/Source/Runtime/TraceLog/Public/Trace/Detail/Atomic.h
--- a/Engine/Source/Runtime/TraceLog/Public/Trace/Detail/Atomic.h
+++ b/Engine/Source/Runtime/TraceLog/Public/Trace/Detail/Atomic.h
@@ -30,6 +30,9 @@ template <typename Type> bool	AtomicCompareExchangeRelease(Type volatile* Target
 template <typename Type> Type	AtomicAddRelaxed(Type volatile* Target, Type Value);
 template <typename Type> Type	AtomicAddRelease(Type volatile* Target, Type Value);
 template <typename Type> Type	AtomicAddAcquire(Type volatile* Target, Type Value);
+template <typename Type> Type	AtomicSubRelaxed(Type volatile* Target, Type Value);
+template <typename Type> Type	AtomicSubRelease(Type volatile* Target, Type Value);
+template <typename Type> Type	AtomicSubAcquire(Type volatile* Target, Type Value);
 void							PlatformYield();
 
 ////////////////////////////////////////////////////////////////////////////////
@@ -147,6 +150,30 @@ inline Type AtomicAddRelease(Type volatile* Target, Type Value)
 	return T->fetch_add(Value, std::memory_order_release);
 }
 
+////////////////////////////////////////////////////////////////////////////////
+template <typename Type>
+inline Type AtomicSubRelaxed(Type volatile* Target, Type Value)
+{
+	std::atomic<Type>* T = (std::atomic<Type>*) Target;
+	return T->fetch_sub(Value, std::memory_order_relaxed);
+}
+
+////////////////////////////////////////////////////////////////////////////////
+template <typename Type>
+inline Type AtomicSubAcquire(Type volatile* Target, Type Value)
+{
+	std::atomic<Type>* T = (std::atomic<Type>*) Target;
+	return T->fetch_sub(Value, std::memory_order_acquire);
+}
+
+////////////////////////////////////////////////////////////////////////////////
+template <typename Type>
+inline Type AtomicSubRelease(Type volatile* Target, Type Value)
+{
+	std::atomic<Type>* T = (std::atomic<Type>*) Target;
+	return T->fetch_sub(Value, std::memory_order_release);
+}
+	
 } // namespace Private
 } // namespace Trace
 } // namespace UE
diff --git a/Engine/Source/Runtime/TraceLog/Public/Trace/Detail/Trace.h b/Engine/Source/Runtime/TraceLog/Public/Trace/Detail/Trace.h
--- a/Engine/Source/Runtime/TraceLog/Public/Trace/Detail/Trace.h
+++ b/Engine/Source/Runtime/TraceLog/Public/Trace/Detail/Trace.h
@@ -6,6 +6,7 @@
 
 
 #include "Misc/Launder.h"
+#include "Misc/CoreMiscDefines.h"
 
 #include <type_traits>
 
diff --git a/Engine/Source/Runtime/TraceLog/Public/Trace/Trace.h b/Engine/Source/Runtime/TraceLog/Public/Trace/Trace.h
--- a/Engine/Source/Runtime/TraceLog/Public/Trace/Trace.h
+++ b/Engine/Source/Runtime/TraceLog/Public/Trace/Trace.h
@@ -147,15 +147,19 @@ struct FMessageEvent
 };
 
 using OnMessageFunc = void(const FMessageEvent&);
+using OnScopeBeginFunc = void(const ANSICHAR*);
+using OnScopeEndFunc = void(void);
 
 struct FInitializeDesc
 {
-	uint32			TailSizeBytes		= 4 << 20; // can be set to 0 to disable the tail buffer
-	uint32			ThreadSleepTimeInMS = 0;
-	bool			bUseWorkerThread	= true;
-	bool			bUseImportantCache	= true;
-	uint32			SessionGuid[4]		= {0,0,0,0}; // leave as zero to generate random
-	OnConnectFunc*	OnConnectionFunc	= nullptr;
+	uint32				TailSizeBytes		= 4 << 20; // can be set to 0 to disable the tail buffer
+	uint32				ThreadSleepTimeInMS = 0;
+	bool				bUseWorkerThread	= true;
+	bool				bUseImportantCache	= true;
+	uint32				SessionGuid[4]		= {0,0,0,0}; // leave as zero to generate random
+	OnConnectFunc*		OnConnectionFunc	= nullptr;
+	OnScopeBeginFunc*	OnScopeBeginFunc	= nullptr;
+	OnScopeEndFunc*		OnScopeEndFunc		= nullptr;
 };
 
 typedef uint32 FChannelId;
@@ -180,12 +184,16 @@ typedef bool		ChannelIterCallback(const FChannelInfo& Info, void*/*User*/);
 
 struct FStatistics
 {
-	uint64		BytesSent		= 0;
-	uint64		BytesTraced		= 0;
-	uint64		MemoryUsed		= 0;
-	uint32		CacheAllocated	= 0;	// Total memory allocated in cache buffers
-	uint32		CacheUsed		= 0;	// Used cache memory; Important-marked events are stored in the cache.
-	uint32		CacheWaste		= 0;	// Unused memory from retired cache buffers
+	uint64	BytesSent				= 0;	// Bytes sent/written to the current endpoint
+	uint64	BytesTraced				= 0;	// Bytes sent/written to the current endpoint (uncompressed)
+	uint64	BytesEmitted			= 0;	// Bytes emitted but potentially not yet written
+	uint64	MemoryUsed				= 0;	// Memory allocated by TraceLog allocator functions
+	uint64	BlockPoolAllocated		= 0;	// Memory allocated for the (TLS) block pool.
+	uint32	SharedBufferAllocated	= 0;	// Memory allocated for shared buffers
+	uint32	FixedBufferAllocated	= 0;	// Memory allocated for fixed buffers (tail, send)
+	uint32	CacheAllocated			= 0;	// Total memory allocated in cache buffers
+	uint32	CacheUsed				= 0;	// Used cache memory; Important-marked events are stored in the cache.
+	uint32	CacheWaste				= 0;	// Unused memory from retired cache buffers
 };
 
 struct FSendFlags
