Revert change to block calculation

diff --git a/Engine/Source/Runtime/TraceLog/Private/Trace/Important/SharedBuffer.cpp b/Engine/Source/Runtime/TraceLog/Private/Trace/Important/SharedBuffer.cpp
--- a/Engine/Source/Runtime/TraceLog/Private/Trace/Important/SharedBuffer.cpp
+++ b/Engine/Source/Runtime/TraceLog/Private/Trace/Important/SharedBuffer.cpp
@@ -125,7 +125,7 @@ static void Writer_RetireSharedBufferImpl()
 
 	FSharedBuffer* Temp = GTailBuffer->Next;
 	const uint32 AllocSizeBytes = GTailBuffer->Size + sizeof(uint32) + sizeof(FSharedBuffer);
-	void* Block = (uint8*)GTailBuffer - AllocSizeBytes;
+	void* Block = (uint8*)GTailBuffer - GTailBuffer->Size - sizeof(uint32);
 	Writer_MemoryFree(Block, AllocSizeBytes);
 #if TRACE_PRIVATE_STATISTICS
 	AtomicSubRelease(&GTraceStatistics.SharedBufferAllocated, AllocSizeBytes);
