[TSAN]
 - Fix a few races in the tracing system

#rnx

diff --git a/Engine/Source/Runtime/TraceLog/Private/Trace/Writer.cpp b/Engine/Source/Runtime/TraceLog/Private/Trace/Writer.cpp
--- a/Engine/Source/Runtime/TraceLog/Private/Trace/Writer.cpp
+++ b/Engine/Source/Runtime/TraceLog/Private/Trace/Writer.cpp
@@ -577,7 +577,7 @@ static bool Writer_UpdateConnection()
 	// Generate Guid for new connection
 	Writer_CreateGuid(&GTraceGuid);
 
-	GDataHandle = PendingDataHandle;
+	AtomicStoreRelaxed(&GDataHandle, PendingDataHandle);
 	if (!Writer_SessionPrologue())
 	{
 		return false;
@@ -683,9 +683,9 @@ void Writer_CallbackOnConnect()
 	// trace data to that tail. We do not want that data to be sent over the
 	// wire as that would cause data to be sent out-of-order.
 	UPTRINT DataHandle = GDataHandle;
-	GDataHandle = 0;
+	AtomicStoreRelaxed(&GDataHandle, (UPTRINT)0);
 	Writer_DrainLocalBuffers();
-	GDataHandle = DataHandle;
+	AtomicStoreRelaxed(&GDataHandle, DataHandle);
 
 	// Issue callback. We assume any events emitted here are not marked as
 	// important and emitted on this thread.
@@ -1218,7 +1218,7 @@ bool Writer_SendSnapshotTo(const ANSICHAR* Host, uint32 Port)
 ////////////////////////////////////////////////////////////////////////////////
 bool Writer_IsTracing()
 {
-	return GDataHandle != 0 || GPendingDataHandle != 0;
+	return AtomicLoadRelaxed(&GDataHandle) != 0 || AtomicLoadRelaxed(&GPendingDataHandle) != 0;
 }
 	
 ////////////////////////////////////////////////////////////////////////////////
diff --git a/Engine/Source/Runtime/TraceLog/Public/Trace/Detail/Trace.h b/Engine/Source/Runtime/TraceLog/Public/Trace/Detail/Trace.h
--- a/Engine/Source/Runtime/TraceLog/Public/Trace/Detail/Trace.h
+++ b/Engine/Source/Runtime/TraceLog/Public/Trace/Detail/Trace.h
@@ -91,7 +91,7 @@ class FChannel;
 		bIsDefinition64 = ((0, ##__VA_ARGS__) & Definition64bit) != 0,}; \
 		typedef std::conditional_t<bIsDefinition8, UE::Trace::FEventRef8, std::conditional_t<bIsDefinition16, UE::Trace::FEventRef16 , std::conditional_t<bIsDefinition64, UE::Trace::FEventRef64, UE::Trace::FEventRef32>>> DefinitionType;\
 		static constexpr uint32 GetSize() { return EventProps_Meta::Size; } \
-		static uint32 TSAN_SAFE GetUid() { static uint32 Uid = 0; return (Uid = Uid ? Uid : Initialize()); } \
+		static uint32 GetUid() { static uint32 Uid { Initialize() }; return Uid; } \
 		static uint32 FORCENOINLINE Initialize() \
 		{ \
 			static const uint32 Uid_ThreadSafeInit = [] () \
