Fix rare crashes when using enabling TraceLog profiling scopes

diff --git a/Engine/Source/Runtime/TraceLog/Private/Trace/Writer.cpp b/Engine/Source/Runtime/TraceLog/Private/Trace/Writer.cpp
--- a/Engine/Source/Runtime/TraceLog/Private/Trace/Writer.cpp
+++ b/Engine/Source/Runtime/TraceLog/Private/Trace/Writer.cpp
@@ -388,7 +388,7 @@ void Writer_SendDataRaw(const void* Data, uint32 Size)
 }
 
 ////////////////////////////////////////////////////////////////////////////////
-void Writer_SendData(uint32 ThreadId, uint8* __restrict Data, uint32 Size)
+void Writer_SendDataNoInstr(uint32 ThreadId, uint8* __restrict Data, uint32 Size)
 {
 	static_assert(ETransport::Active == ETransport::TidPacketSync, "Active should be set to what the compiled code uses. It is used to track places that assume transport packet format");
 
@@ -397,7 +397,6 @@ void Writer_SendData(uint32 ThreadId, uint8* __restrict Data, uint32 Size)
 		return;
 	}
 	
-	FProfilerScope _(__func__);
 
 	// Smaller buffers usually aren't redundant enough to benefit from being
 	// compressed. They often end up being larger.
@@ -442,9 +441,19 @@ void Writer_SendData(uint32 ThreadId, uint8* __restrict Data, uint32 Size)
 #endif
 }
 
+////////////////////////////////////////////////////////////////////////////////
+void Writer_SendData(uint32 ThreadId, uint8* __restrict Data, uint32 Size)
+{
+	FProfilerScope _(__func__);
+	Writer_SendDataNoInstr(ThreadId, Data, Size);
+}
+	
 ////////////////////////////////////////////////////////////////////////////////
 static void Writer_DescribeEvents(FEventNode::FIter Iter)
 {
+	// Since we are using on stack buffer to collect the describing events we
+	// cannot emit additional data on this thread. So use the un-instrumented
+	// send data functions here.
 	TWriteBufferRedirect<4096> TraceData;
 
 	while (const FEventNode* Event = Iter.GetNext())
@@ -454,14 +463,14 @@ static void Writer_DescribeEvents(FEventNode::FIter Iter)
 		// Flush just in case an NewEvent event will be larger than 512 bytes.
 		if (TraceData.GetSize() >= (TraceData.GetCapacity() - 512))
 		{
-			Writer_SendData(ETransportTid::Events, TraceData.GetData(), TraceData.GetSize());
+			Writer_SendDataNoInstr(ETransportTid::Events, TraceData.GetData(), TraceData.GetSize());
 			TraceData.Reset();
 		}
 	}
 
 	if (TraceData.GetSize())
 	{
-		Writer_SendData(ETransportTid::Events, TraceData.GetData(), TraceData.GetSize());
+		Writer_SendDataNoInstr(ETransportTid::Events, TraceData.GetData(), TraceData.GetSize());
 	}
 }
 
