Add support for relay endpoints in TraceLog

Adds an additional method of writing traces where user can supply their own open/write/close functions. Snapshot to relay connection is not yet supported.

diff --git a/Engine/Source/Runtime/Core/Private/ProfilingDebugging/TraceAuxiliary.cpp b/Engine/Source/Runtime/Core/Private/ProfilingDebugging/TraceAuxiliary.cpp
--- a/Engine/Source/Runtime/Core/Private/ProfilingDebugging/TraceAuxiliary.cpp
+++ b/Engine/Source/Runtime/Core/Private/ProfilingDebugging/TraceAuxiliary.cpp
@@ -122,6 +122,7 @@ public:
 	void DisableChannels(const TCHAR* ChannelList, const FTraceAuxiliary::FLogCategoryAlias& LogCategory);
 	void DisableChannels(TConstArrayView<uint32> ChannelIds);
 	bool Connect(FTraceAuxiliary::EConnectionType Type, const TCHAR* Parameter, const FTraceAuxiliary::FLogCategoryAlias& LogCategory, uint16 SendFlags);
+	bool ConnectRelay(UPTRINT Handle, UE::Trace::IoWriteFunc WriteFunc, UE::Trace::IoCloseFunc CloseFunc, uint16 SendFlags);
 	bool Stop();
 	void FreezeReadOnlyChannels();
 	void ResumeChannels();
@@ -461,6 +462,39 @@ bool FTraceAuxiliaryImpl::Connect(FTraceAuxiliary::EConnectionType Type, const T
 	return bConnected;
 }
 
+////////////////////////////////////////////////////////////////////////////////
+bool FTraceAuxiliaryImpl::ConnectRelay(UPTRINT Handle, UE::Trace::IoWriteFunc WriteFunc, UE::Trace::IoCloseFunc CloseFunc, uint16 SendFlags)
+{
+	if (UE::Trace::IsTracing())
+	{
+		UE_LOG_REF(LogTrace, Error, TEXT("Already tracing from unknown source."));
+		return true;
+	}
+	
+	const bool bConnected = UE::Trace::RelayTo(Handle, WriteFunc, CloseFunc);
+	if (!bConnected)
+	{
+		UE_LOG(LogTrace, Display, TEXT("Trace started (connected to relay endpoint)"));
+		
+		FString StartedDest = FString::Printf(TEXT("Relay endpoint (Handle: 0x%llx)"), Handle);
+		FTraceAuxiliary::EConnectionType StartedType = FTraceAuxiliary::EConnectionType::Relay;
+		
+		{
+			FWriteScopeLock _(CurrentTargetLock);
+			CurrentTraceTarget.TraceType = StartedType;
+			CurrentTraceTarget.TraceDest = StartedDest;
+		}
+		
+		TRACE_CPUPROFILER_EVENT_SCOPE(FTraceAux_OnTraceStarted);
+		FTraceAuxiliary::OnTraceStarted.Broadcast(StartedType, StartedDest);
+	}
+	else
+	{
+		UE_LOG(LogTrace, Error, TEXT("Trace failed to connect relay endpoint"));
+	}
+	return bConnected;
+}
+
 ////////////////////////////////////////////////////////////////////////////////
 bool FTraceAuxiliaryImpl::Stop()
 {
@@ -1468,15 +1502,15 @@ FTraceAuxiliary::FOnTraceStopped FTraceAuxiliary::OnTraceStopped;
 FTraceAuxiliary::FOnSnapshotSaved FTraceAuxiliary::OnSnapshotSaved;
 
 ////////////////////////////////////////////////////////////////////////////////
-bool FTraceAuxiliary::Start(EConnectionType Type, const TCHAR* Target, const TCHAR* Channels, FOptions* Options, const FLogCategoryAlias& LogCategory)
-{
 #if UE_TRACE_ENABLED
+
+static bool TraceAuxiliaryStartShared(const TCHAR* Channels, const FTraceAuxiliary::FOptions* Options, const FTraceAuxiliary::FLogCategoryAlias& LogCategory)
+{
 	if (GTraceAuxiliary.IsParentProcessAndPreFork())
 	{
 		return false;
 	}
 
-	TRACE_CPUPROFILER_EVENT_SCOPE(FTraceAux_Start);
 
 	if (GTraceAuxiliary.IsConnected())
 	{
@@ -1502,6 +1536,25 @@ bool FTraceAuxiliary::Start(EConnectionType Type, const TCHAR* Target, const TCH
 		GTraceAuxiliary.EnableCommandlineChannels();
 	}
 
+	return true;
+}
+
+#endif // UE_TRACE_ENABLED
+
+////////////////////////////////////////////////////////////////////////////////
+bool FTraceAuxiliary::Start(EConnectionType Type, const TCHAR* Target, const TCHAR* Channels, FOptions* Options, const FLogCategoryAlias& LogCategory)
+{
+#if UE_TRACE_ENABLED
+	// Use StartCustom function to start a custom trace
+	check(Type != EConnectionType::Relay);
+	
+	TRACE_CPUPROFILER_EVENT_SCOPE(FTraceAux_Start);
+
+	if (!TraceAuxiliaryStartShared(Channels, Options, LogCategory))
+	{
+		return false;
+	}
+
 	if (Options)
 	{
 		// Truncation is only valid when tracing to file and filename is set.
@@ -1519,6 +1572,24 @@ bool FTraceAuxiliary::Start(EConnectionType Type, const TCHAR* Target, const TCH
 #endif
 }
 
+////////////////////////////////////////////////////////////////////////////////
+bool FTraceAuxiliary::Relay(UPTRINT Handle, UE::Trace::IoWriteFunc WriteFunc, UE::Trace::IoCloseFunc CloseFunc, const TCHAR* Channels, const FOptions* Options)
+{
+#if UE_TRACE_ENABLED
+	TRACE_CPUPROFILER_EVENT_SCOPE(FTraceAux_Start);
+	
+	if (!TraceAuxiliaryStartShared(Channels, Options, LogTrace))
+	{
+		return false;
+	}
+	
+	uint16 SendFlags = (Options && Options->bExcludeTail) ? UE::Trace::FSendFlags::ExcludeTail : 0;
+	return GTraceAuxiliary.ConnectRelay(Handle, WriteFunc, CloseFunc, SendFlags);
+#else
+	return false;
+#endif
+}
+
 ////////////////////////////////////////////////////////////////////////////////
 bool FTraceAuxiliary::Stop()
 {
diff --git a/Engine/Source/Runtime/Core/Public/ProfilingDebugging/TraceAuxiliary.h b/Engine/Source/Runtime/Core/Public/ProfilingDebugging/TraceAuxiliary.h
--- a/Engine/Source/Runtime/Core/Public/ProfilingDebugging/TraceAuxiliary.h
+++ b/Engine/Source/Runtime/Core/Public/ProfilingDebugging/TraceAuxiliary.h
@@ -83,6 +83,10 @@ public:
 		 * If target is null the current date and time is used.
 		 */
 		File,
+		/**
+		 * Relay connection. Pass user defined I/O functions to write raw trace data.
+		 */
+		Relay,
 		/**
 		 * Don't connect, just start tracing to memory.
 		 */
@@ -123,7 +127,7 @@ public:
 	/**
 	 * Start tracing to a target (network connection or file) with an active set of channels. If a connection is
 	 * already active this call does nothing.
-	 * @param Type Type of connection
+	 * @param Type Type of connection. Network or File type. If a custom target is desired use \ref FTraceAuxiliary::Relay.
 	 * @param Target String to use for connection. See /ref EConnectionType for details.
 	 * @param Channels Comma separated list of channels to enable. Default set of channels are enabled if argument is not specified. If the pointer is null no channels are enabled.
 	 * @param Options Optional additional tracing options.
@@ -132,6 +136,17 @@ public:
 	 */
 	static CORE_API bool Start(EConnectionType Type, const TCHAR* Target, const TCHAR* Channels = TEXT("default"), FOptions* Options = nullptr, const FLogCategoryAlias& LogCategory = LogCore);
 
+	/**
+	 * Start tracing to a custom target by providing a handle and a write and close function. If a connection is already active this call does nothing.
+	 * @param Handle Anonymous handle to trace to. This will be passed to the writer and close functions.
+	 * @param WriteFunc A function to handle writing of trace data
+	 * @param CloseFunc A function to handle closing of the data stream.
+	 * @param Channels Comma separated list of channels to enable. Default set of channels are enabled if argument is not specified. If the pointer is null no channels are enabled.
+	 * @param Options Optional additional tracing options.
+	 * @return True when successfully starting the trace, false if the data connection could not be made.
+	 */
+	static CORE_API bool Relay(UPTRINT Handle, UE::Trace::IoWriteFunc WriteFunc, UE::Trace::IoCloseFunc CloseFunc, const TCHAR* Channels = TEXT("default"), const FOptions* Options = nullptr);
+
 	/**
 	 * Stop tracing.
 	 * @return True if the trace was stopped, false if there was no data connection.
diff --git a/Engine/Source/Runtime/TraceLog/Private/Trace/Trace.cpp b/Engine/Source/Runtime/TraceLog/Private/Trace/Trace.cpp
--- a/Engine/Source/Runtime/TraceLog/Private/Trace/Trace.cpp
+++ b/Engine/Source/Runtime/TraceLog/Private/Trace/Trace.cpp
@@ -23,6 +23,7 @@ void	Writer_Shutdown();
 void	Writer_Update();
 bool	Writer_SendTo(const ANSICHAR*, uint32, uint32);
 bool	Writer_WriteTo(const ANSICHAR*, uint32);
+bool	Writer_RelayTo(UPTRINT, IoWriteFunc, IoCloseFunc, uint16);
 bool	Writer_WriteSnapshotTo(const ANSICHAR*);
 bool	Writer_SendSnapshotTo(const ANSICHAR*, uint32);	
 bool	Writer_IsTracing();
@@ -112,6 +113,12 @@ bool WriteTo(const TCHAR* InPath, uint16 Flags)
 	return Private::Writer_WriteTo(Path, Flags);
 }
 
+////////////////////////////////////////////////////////////////////////////////
+bool RelayTo(UPTRINT InHandle, IoWriteFunc WriteFunc, IoCloseFunc CloseFunc, uint16 Flags)
+{
+	return Private::Writer_RelayTo(InHandle, WriteFunc, CloseFunc, Flags);
+}
+
 ////////////////////////////////////////////////////////////////////////////////
 bool WriteSnapshotTo(const TCHAR* InPath)
 {
diff --git a/Engine/Source/Runtime/TraceLog/Private/Trace/Writer.cpp b/Engine/Source/Runtime/TraceLog/Private/Trace/Writer.cpp
--- a/Engine/Source/Runtime/TraceLog/Private/Trace/Writer.cpp
+++ b/Engine/Source/Runtime/TraceLog/Private/Trace/Writer.cpp
@@ -87,7 +87,15 @@ static uint32					GUpdateCounter;		// = 0;
 uint64							GPacketSerial = 1;
 #endif
 
+////////////////////////////////////////////////////////////////////////////////
+struct FWriterState
+{
+	IoWriteFunc* Write;
+	IoCloseFunc* Close;
+};
 
+static FWriterState GWriterState;
+static FWriterState GPendingWriterState;
 
 ////////////////////////////////////////////////////////////////////////////////
 // When a thread terminates we want to recover its trace buffer. To do this we
@@ -317,10 +325,10 @@ static bool Writer_FlushSendBuffer()
 {
 	if( GSendBufferCursor > GSendBuffer )
 	{
-		if (!IoWrite(GDataHandle, GSendBuffer, GSendBufferCursor - GSendBuffer))
+		if (!GWriterState.Write(GDataHandle, GSendBuffer, GSendBufferCursor - GSendBuffer))
 		{
 			UE_TRACE_ERRORMESSAGE(WriteError, GetLastErrorCode());
-			IoClose(GDataHandle);
+			GWriterState.Close(GDataHandle);
 			GDataHandle = 0;
 			return false;
 		}
@@ -354,9 +362,9 @@ static void Writer_SendDataImpl(const void* Data, uint32 Size)
 	// Should rarely happen but if we're asked to send large data send it directly
 	if (Size > GSendBufferSize)
 	{
-		if (!IoWrite(GDataHandle, Data, Size))
+		if (!GWriterState.Write(GDataHandle, Data, Size))
 		{
-			IoClose(GDataHandle);
+			GWriterState.Close(GDataHandle);
 			GDataHandle = 0;
 		}
 	}
@@ -367,10 +375,10 @@ static void Writer_SendDataImpl(const void* Data, uint32 Size)
 		GSendBufferCursor += Size;
 	}
 #else
-	if (!IoWrite(GDataHandle, Data, Size))
+	if (!GWriterState.Write(GDataHandle, Data, Size))
 	{
 		UE_TRACE_ERRORMESSAGE(WriteError, GetLastErrorCode());
-		IoClose(GDataHandle);
+		GWriterState.Close(GDataHandle);
 		GDataHandle = 0;
 	}
 #endif
@@ -534,7 +542,7 @@ static void Writer_Close()
 	if (GDataHandle)
 	{
 		Writer_FlushSendBuffer();
-		IoClose(GDataHandle);
+		GWriterState.Close(GDataHandle);
 	}
 
 	GDataHandle = 0;
@@ -570,6 +578,7 @@ static bool Writer_UpdateConnection()
 		return true;
 	}
 
+	FWriterState PendingWriterState = GPendingWriterState;
 	AtomicStoreRelaxed(&GPendingDataHandle, UPTRINT(0));
 
 	// Extract send flags
@@ -579,14 +588,15 @@ static bool Writer_UpdateConnection()
 	// Reject the pending connection if we've already got a connection
 	if (GDataHandle)
 	{
-		IoClose(PendingDataHandle);
+		GWriterState.Close(PendingDataHandle);
 		return false;
 	}
 
 	// Generate Guid for new connection
 	Writer_CreateGuid(&GTraceGuid);
 
-	AtomicStoreRelaxed(&GDataHandle, PendingDataHandle);
+	GWriterState = PendingWriterState;
+	AtomicStoreRelease(&GDataHandle, PendingDataHandle);
 	if (!Writer_SessionPrologue())
 	{
 		return false;
@@ -660,19 +670,19 @@ static bool Writer_SessionPrologue()
 	FHandshake Handshake;
 	memcpy(&Handshake.SessionGuid, &GSessionGuid, sizeof(FTraceGuid));
 	memcpy(&Handshake.TraceGuid, &GTraceGuid, sizeof(FTraceGuid));
-	bool bOk = IoWrite(GDataHandle, &Handshake, FHandshake::Size);
+	bool bOk = GWriterState.Write(GDataHandle, &Handshake, FHandshake::Size);
 
 	// Stream header
 	const struct {
 		uint8 TransportVersion	= ETransport::TidPacketSync;
 		uint8 ProtocolVersion	= EProtocol::Id;
 	} TransportHeader;
-	bOk &= IoWrite(GDataHandle, &TransportHeader, sizeof(TransportHeader));
+	bOk &= GWriterState.Write(GDataHandle, &TransportHeader, sizeof(TransportHeader));
 
 	if (!bOk)
 	{
 		UE_TRACE_ERRORMESSAGE(WriteError, GetLastErrorCode());
-		IoClose(GDataHandle);
+		GWriterState.Close(GDataHandle);
 		GDataHandle = 0;
 		return false;
 	}
@@ -821,7 +831,7 @@ static void Writer_InternalShutdown()
 	if (GDataHandle)
 	{
 		Writer_FlushSendBuffer();
-		IoClose(GDataHandle);
+		GWriterState.Close(GDataHandle);
 		GDataHandle = 0;
 	}
 
@@ -944,7 +954,7 @@ static UPTRINT Writer_PackSendFlags(UPTRINT DataHandle, uint32 Flags)
 	// Guard against the assumption being wrong.
 	if (DataHandle & 0xffff'0000'0000'0000ull)
 	{
-		IoClose(DataHandle);
+		GWriterState.Close(DataHandle);
 		return 0;
 	}
 
@@ -977,7 +987,9 @@ bool Writer_SendTo(const ANSICHAR* Host, uint32 Flags, uint32 Port)
 		return false;
 	}
 
-	AtomicStoreRelaxed(&GPendingDataHandle, DataHandle);
+	GPendingWriterState = { &IoWrite, &IoClose };
+	
+	AtomicStoreRelease(&GPendingDataHandle, DataHandle);
 	return true;
 #else
 	return false;
@@ -1008,13 +1020,38 @@ bool Writer_WriteTo(const ANSICHAR* Path, uint32 Flags)
 		return false;
 	}
 
-	AtomicStoreRelaxed(&GPendingDataHandle, DataHandle);
+	GPendingWriterState = { &IoWrite, &IoClose };
+	
+	AtomicStoreRelease(&GPendingDataHandle, DataHandle);
 	return true;
 #else
 	return false;
 #endif
 }
 
+////////////////////////////////////////////////////////////////////////////////
+bool Writer_RelayTo(UPTRINT InHandle, IoWriteFunc WriteFunc, IoCloseFunc CloseFunc, uint16 Flags)
+{
+	if (AtomicLoadRelaxed(&GPendingDataHandle))
+	{
+		return false;
+	}
+	
+	Writer_InternalInitialize();
+	
+	UPTRINT DataHandle = InHandle;
+	DataHandle = Writer_PackSendFlags(DataHandle, Flags);
+	if (!DataHandle)
+	{
+		return false;
+	}
+
+	GPendingWriterState = { WriteFunc, CloseFunc };
+	
+	AtomicStoreRelease(&GPendingDataHandle, DataHandle);
+	return true;
+}
+
 ////////////////////////////////////////////////////////////////////////////////
 struct WorkerUpdateLock
 {
@@ -1142,6 +1179,8 @@ bool Writer_WriteSnapshot(const FSnapshotTarget& Target)
 		const bool bExistingDataHandle = GDataHandle != 0;
 		
 		TStashGlobal DataHandle(GDataHandle);
+		TStashGlobal StateWrite(GWriterState.Write);
+		TStashGlobal StateClose(GWriterState.Close);
 		TStashGlobal PendingDataHandle(GPendingDataHandle);
 		TStashGlobal SyncPacketCountdown(GSyncPacketCountdown, GNumSyncPackets);
 		TStashGlobal TraceStatistics(GTraceStatistics);
diff --git a/Engine/Source/Runtime/TraceLog/Public/Trace/Trace.h b/Engine/Source/Runtime/TraceLog/Public/Trace/Trace.h
--- a/Engine/Source/Runtime/TraceLog/Public/Trace/Trace.h
+++ b/Engine/Source/Runtime/TraceLog/Public/Trace/Trace.h
@@ -181,6 +181,8 @@ typedef void		ChannelIterFunc(const ANSICHAR*, bool, void*);
 /*  The callback provides information about a channel and a user provided pointer. 
 	Returning false from the callback will stop the enumeration */
 typedef bool		ChannelIterCallback(const FChannelInfo& Info, void*/*User*/);
+typedef bool		IoWriteFunc(UPTRINT Handle, const void* Data, uint32 Size);
+typedef void		IoCloseFunc(UPTRINT Handle);	
 
 struct FStatistics
 {
@@ -213,6 +215,7 @@ UE_TRACE_API void	Update() UE_TRACE_IMPL();
 UE_TRACE_API void	GetStatistics(FStatistics& Out) UE_TRACE_IMPL();
 UE_TRACE_API bool	SendTo(const TCHAR* Host, uint32 Port=0, uint16 Flags=FSendFlags::None) UE_TRACE_IMPL(false);
 UE_TRACE_API bool	WriteTo(const TCHAR* Path, uint16 Flags=FSendFlags::None) UE_TRACE_IMPL(false);
+UE_TRACE_API bool	RelayTo(UPTRINT InHandle, IoWriteFunc WriteFunc, IoCloseFunc CloseFunc, uint16 Flags = FSendFlags::None) UE_TRACE_IMPL(false);
 UE_TRACE_API bool	WriteSnapshotTo(const TCHAR* Path) UE_TRACE_IMPL(false);
 UE_TRACE_API bool	SendSnapshotTo(const TCHAR* Host, uint32 Port) UE_TRACE_IMPL(false);	
 UE_TRACE_API bool	IsTracing() UE_TRACE_IMPL(false);
