[Insights] Fixed TSAN data race warning re usage of UE::Trace::SetUpdateCallback(). Also normalized function arguments receiving a pointer to function (for OnMessageFunc and OnUpdateFunc callbacks).

#rnx

diff --git a/Engine/Source/Runtime/TraceLog/Private/Trace/Message.cpp b/Engine/Source/Runtime/TraceLog/Private/Trace/Message.cpp
--- a/Engine/Source/Runtime/TraceLog/Private/Trace/Message.cpp
+++ b/Engine/Source/Runtime/TraceLog/Private/Trace/Message.cpp
@@ -12,9 +12,9 @@ namespace Private {
 
 ////////////////////////////////////////////////////////////////////////////////
 static OnMessageFunc* GMessageFunc = nullptr;
-	
+
 ////////////////////////////////////////////////////////////////////////////////
-void Message_SetCallback(OnMessageFunc Callback)
+void Message_SetCallback(OnMessageFunc* Callback)
 {
 	GMessageFunc = Callback;
 }
diff --git a/Engine/Source/Runtime/TraceLog/Private/Trace/Message.h b/Engine/Source/Runtime/TraceLog/Private/Trace/Message.h
--- a/Engine/Source/Runtime/TraceLog/Private/Trace/Message.h
+++ b/Engine/Source/Runtime/TraceLog/Private/Trace/Message.h
@@ -11,8 +11,8 @@ namespace UE {
 namespace Trace {
 namespace Private {
 
-bool GetErrorMessage(char* OutBuffer, uint32 BufferSize, int32 ErrorCode);	
-	
+bool GetErrorMessage(char* OutBuffer, uint32 BufferSize, int32 ErrorCode);
+
 /**
  * Utility macro to format messages
  */
@@ -22,21 +22,21 @@ bool GetErrorMessage(char* OutBuffer, uint32 BufferSize, int32 ErrorCode);
 
 #define UE_TRACE_MESSAGE(Type, Msg) \
 	Message_Send(EMessageType::Type, #Type, Msg);
-	
+
 #define UE_TRACE_MESSAGE_F(Type, Fmt, ...) \
 	{\
 		char Buff[UE_TRACE_MESSAGE_FMT_MAX_SIZE];\
 		snprintf(Buff, UE_TRACE_MESSAGE_FMT_MAX_SIZE, Fmt, __VA_ARGS__);\
 		Message_Send(EMessageType::Type, #Type, Buff);\
 	}
-	
+
 #define UE_TRACE_ERRORMESSAGE(Type, ErrorCode)\
 	{\
 		char ErrorMessageBuffer[UE_TRACE_MESSAGE_ERR_MAX_SIZE] = {'\0'};\
 		GetErrorMessage(ErrorMessageBuffer, UE_TRACE_MESSAGE_ERR_MAX_SIZE, ErrorCode);\
 		UE_TRACE_MESSAGE_F(Type, "(error code %d): '%s'", ErrorCode, ErrorMessageBuffer)\
 	}
-	
+
 #define UE_TRACE_ERRORMESSAGE_F(Type, ErrorCode, ContextFmt, ...)\
 	{\
 		char ErrorMessageBuffer[UE_TRACE_MESSAGE_ERR_MAX_SIZE];\
@@ -55,14 +55,14 @@ bool GetErrorMessage(char* OutBuffer, uint32 BufferSize, int32 ErrorCode);
 /**
  * Initialize messages. Saves the provided callback if specified.
  */
-void Message_SetCallback(UE::Trace::OnMessageFunc Callback);
-	
+void Message_SetCallback(UE::Trace::OnMessageFunc* Callback);
+
 /**
  * Sends a message to the registered callback. Description string is not
  * expected to outlive the function call.
  */
 void Message_Send(EMessageType Type, const char* TypeStr, const char* Description = nullptr);
-	
+
 
 } } } // namespace UE::Trace::Private
 
diff --git a/Engine/Source/Runtime/TraceLog/Private/Trace/Trace.cpp b/Engine/Source/Runtime/TraceLog/Private/Trace/Trace.cpp
--- a/Engine/Source/Runtime/TraceLog/Private/Trace/Trace.cpp
+++ b/Engine/Source/Runtime/TraceLog/Private/Trace/Trace.cpp
@@ -15,8 +15,8 @@ namespace Private
 {
 
 ////////////////////////////////////////////////////////////////////////////////
-void	Message_SetCallback(OnMessageFunc Callback);
-void	Writer_SetUpdateCallback(OnUpdateFunc Callback);
+void	Message_SetCallback(OnMessageFunc* Callback);
+void	Writer_SetUpdateCallback(OnUpdateFunc* Callback);
 void	Writer_MemorySetHooks(AllocFunc, FreeFunc);
 void	Writer_Initialize(const FInitializeDesc&);
 void	Writer_WorkerCreate();
@@ -62,13 +62,13 @@ void SetMemoryHooks(AllocFunc Alloc, FreeFunc Free)
 }
 
 ////////////////////////////////////////////////////////////////////////////////
-void SetMessageCallback(OnMessageFunc MessageFunc)
+void SetMessageCallback(OnMessageFunc* MessageFunc)
 {
 	Private::Message_SetCallback(MessageFunc);
 }
 
 ////////////////////////////////////////////////////////////////////////////////
-void SetUpdateCallback(OnUpdateFunc UpdateFunc)
+void SetUpdateCallback(OnUpdateFunc* UpdateFunc)
 {
 	Private::Writer_SetUpdateCallback(UpdateFunc);
 }
diff --git a/Engine/Source/Runtime/TraceLog/Private/Trace/Writer.cpp b/Engine/Source/Runtime/TraceLog/Private/Trace/Writer.cpp
--- a/Engine/Source/Runtime/TraceLog/Private/Trace/Writer.cpp
+++ b/Engine/Source/Runtime/TraceLog/Private/Trace/Writer.cpp
@@ -37,7 +37,7 @@ namespace Private {
 
 ////////////////////////////////////////////////////////////////////////////////
 int32			Encode(const void*, int32, void*, int32);
-void			Writer_SetUpdateCallback(OnUpdateFunc Callback);
+void			Writer_SetUpdateCallback(OnUpdateFunc* Callback);
 void			Writer_SendData(uint32, uint8* __restrict, uint32);
 void			Writer_InitializeTail(int32);
 void			Writer_ShutdownTail();
@@ -722,9 +722,9 @@ static uint32			GSleepTimeInMS = 17;
 static volatile uint32	GUpdateInProgress = 1;	// Don't allow updates until initialized
 
 ////////////////////////////////////////////////////////////////////////////////
-void Writer_SetUpdateCallback(OnUpdateFunc Callback)
+void Writer_SetUpdateCallback(OnUpdateFunc* Callback)
 {
-	GOnUpdate = Callback;
+	AtomicStoreRelaxed(&GOnUpdate, Callback);
 }
 
 ////////////////////////////////////////////////////////////////////////////////
@@ -747,7 +747,7 @@ static void Writer_WorkerUpdateInternal()
 	}
 #endif
 
-	if (OnUpdateFunc* OnUpdate = GOnUpdate)
+	if (OnUpdateFunc* OnUpdate = AtomicLoadRelaxed(&GOnUpdate))
 	{
 		OnUpdate();
 	}
@@ -937,7 +937,7 @@ void Writer_Initialize(const FInitializeDesc& Desc)
 
 	// Store callback on connection
 	GOnConnection = Desc.OnConnectionFunc;
-	GOnUpdate = Desc.OnUpdateFunc;
+	AtomicStoreRelaxed(&GOnUpdate, Desc.OnUpdateFunc);
 
 	// Allow the worker thread to start updating
 	AtomicStoreRelease(&GUpdateInProgress, uint32(0));
diff --git a/Engine/Source/Runtime/TraceLog/Public/Trace/Trace.h b/Engine/Source/Runtime/TraceLog/Public/Trace/Trace.h
--- a/Engine/Source/Runtime/TraceLog/Public/Trace/Trace.h
+++ b/Engine/Source/Runtime/TraceLog/Public/Trace/Trace.h
@@ -207,8 +207,8 @@ struct FSendFlags
 };
 
 UE_TRACE_API void	SetMemoryHooks(AllocFunc Alloc, FreeFunc Free) UE_TRACE_IMPL();
-UE_TRACE_API void	SetMessageCallback(OnMessageFunc MessageFunc) UE_TRACE_IMPL();
-UE_TRACE_API void	SetUpdateCallback(OnUpdateFunc UpdateFunc) UE_TRACE_IMPL();
+UE_TRACE_API void	SetMessageCallback(OnMessageFunc* MessageFunc) UE_TRACE_IMPL();
+UE_TRACE_API void	SetUpdateCallback(OnUpdateFunc* UpdateFunc) UE_TRACE_IMPL();
 UE_TRACE_API void	Initialize(const FInitializeDesc& Desc) UE_TRACE_IMPL();
 UE_TRACE_API void	StartWorkerThread() UE_TRACE_IMPL();
 UE_TRACE_API void	Shutdown() UE_TRACE_IMPL();
