[Insights] Fixed computation for GTraceStatistics.CacheAllocated drifting with 32 bytes on each buffer allocated.

#rnx

diff --git a/Engine/Source/Runtime/TraceLog/Private/Trace/Important/Cache.cpp b/Engine/Source/Runtime/TraceLog/Private/Trace/Important/Cache.cpp
--- a/Engine/Source/Runtime/TraceLog/Private/Trace/Important/Cache.cpp
+++ b/Engine/Source/Runtime/TraceLog/Private/Trace/Important/Cache.cpp
@@ -51,11 +51,11 @@ extern FStatistics		GTraceStatistics;
 ////////////////////////////////////////////////////////////////////////////////
 static FCacheBuffer* Writer_CacheCreateBuffer(uint32 Size)
 {
+	const uint32 CacheBufferSize = sizeof(FCacheBuffer) + Size;
+	void* Block = Writer_MemoryAllocate(CacheBufferSize, alignof(FCacheBuffer));
 #if TRACE_PRIVATE_STATISTICS
-	GTraceStatistics.CacheAllocated += Size;
+	GTraceStatistics.CacheAllocated += CacheBufferSize;
 #endif
-
-	void* Block = Writer_MemoryAllocate(sizeof(FCacheBuffer) + Size, alignof(FCacheBuffer));
 	auto* Buffer = (FCacheBuffer*)Block;
 	Buffer->Size = Size;
 	Buffer->Remaining = Buffer->Size;
