Fix crash when making snapshot

diff --git a/Engine/Source/Runtime/TraceLog/Private/Trace/Writer.cpp b/Engine/Source/Runtime/TraceLog/Private/Trace/Writer.cpp
--- a/Engine/Source/Runtime/TraceLog/Private/Trace/Writer.cpp
+++ b/Engine/Source/Runtime/TraceLog/Private/Trace/Writer.cpp
@@ -91,8 +91,8 @@ uint64							GPacketSerial = 1;
 ////////////////////////////////////////////////////////////////////////////////
 struct FWriterState
 {
-	IoWriteFunc* Write;
-	IoCloseFunc* Close;
+	IoWriteFunc* Write = nullptr;
+	IoCloseFunc* Close = nullptr;
 };
 
 static FWriterState GWriterState;
@@ -1201,6 +1201,8 @@ bool Writer_WriteSnapshot(const FSnapshotTarget& Target)
 		TStashGlobal SyncPacketCountdown(GSyncPacketCountdown, GNumSyncPackets);
 		TStashGlobal TraceStatistics(GTraceStatistics);
 
+		GWriterState = { &IoWrite, &IoClose };
+
 		if (Target.Type == FSnapshotTarget::EType::FileTarget)
 		{
 #if TRACE_PRIVATE_ALLOW_FILE
