Auto flush when scope depth reaches zero. Change transfered from Zen server codebase.

#rnx

diff --git a/Engine/Source/Runtime/TraceLog/standalone_epilogue.h b/Engine/Source/Runtime/TraceLog/standalone_epilogue.h
--- a/Engine/Source/Runtime/TraceLog/standalone_epilogue.h
+++ b/Engine/Source/Runtime/TraceLog/standalone_epilogue.h
@@ -281,6 +281,7 @@ private:
 	uint64		PrevTimestamp = 0;
 	uint8*		Cursor = Buffer;
 	uint32		ThreadIdOverride = 0;
+	uint16		Depth = 0;
 	uint8		Buffer[BufferSize];
 };
 
@@ -292,7 +293,7 @@ void ScopeBuffer::Flush(bool Force)
 	if (Cursor == Buffer)
 		return;
 
-	if (!Force && (Cursor <= (Buffer + BufferSize - Overflow)))
+	if (Depth > 0 && !Force && (Cursor <= (Buffer + BufferSize - Overflow)))
 		return;
 
 	if (ThreadIdOverride)
@@ -313,6 +314,7 @@ void ScopeBuffer::Enter(uint64 Timestamp, uint32 ScopeId, int32 Flags)
 	PrevTimestamp += Timestamp;
 	Cursor += encode64_7bit((Timestamp) << 1 | EnterLsb, Cursor);
 	Cursor += encode32_7bit(ScopeId, Cursor);
+	Depth++;
 
 	bool ShouldFlush = (Flags & CpuScopeFlags::CpuFlush);
 	Flush(ShouldFlush);
@@ -324,6 +326,7 @@ void ScopeBuffer::Leave(uint64 Timestamp)
 	Timestamp -= PrevTimestamp;
 	PrevTimestamp += Timestamp;
 	Cursor += encode64_7bit((Timestamp << 1) | LeaveLsb, Cursor);
+	Depth--;
 
 	Flush(false);
 }
