Workaround Horde trying to compile files that aren't part of any build

#rnx

diff --git a/Engine/Source/Runtime/TraceLog/create_standalone.py b/Engine/Source/Runtime/TraceLog/create_standalone.py
--- a/Engine/Source/Runtime/TraceLog/create_standalone.py
+++ b/Engine/Source/Runtime/TraceLog/create_standalone.py
@@ -153,11 +153,11 @@ def _main_trace(src_dir, dest_dir, thin, analysis):
 	source_files = _finalize_source(source_files, include_dirs)
 
 	# Add prologue and epilogue files
-	prologue = _SourceFile(src_dir / "standalone_prologue.h")
+	prologue = _SourceFile(src_dir / "standalone_prologue.h.src")
 	source_files.insert(0, prologue)
 	prologue.lines = [x for x in prologue.path.open("rt") if x]
 
-	epilogue = _SourceFile(src_dir / "standalone_epilogue.h")
+	epilogue = _SourceFile(src_dir / "standalone_epilogue.h.src")
 	source_files.append(epilogue)
 	epilogue.lines = [x for x in epilogue.path.open("rt") if x]
 
diff --git a/Engine/Source/Runtime/TraceLog/standalone_epilogue.h b/Engine/Source/Runtime/TraceLog/standalone_epilogue.h
deleted file mode 100644
diff --git a/Engine/Source/Runtime/TraceLog/standalone_prologue.h b/Engine/Source/Runtime/TraceLog/standalone_prologue.h
deleted file mode 100644
