[AutoRTFM] Re-enable AutoRTFM for TraceLog

This was disabled in 38699603, but the CL does not explain why.
Re-enabling this gets the verse presubmit tests passing again.

diff --git a/Engine/Source/Runtime/TraceLog/TraceLog.Build.cs b/Engine/Source/Runtime/TraceLog/TraceLog.Build.cs
--- a/Engine/Source/Runtime/TraceLog/TraceLog.Build.cs
+++ b/Engine/Source/Runtime/TraceLog/TraceLog.Build.cs
@@ -14,7 +14,5 @@ public class TraceLog : ModuleRules
 		PublicIncludePathModuleNames.Add("Core");
 
 		PrivateDefinitions.Add("SUPPRESS_PER_MODULE_INLINE_FILE"); // This module does not use core's standard operator new/delete overloads
-
-		bDisableAutoRTFMInstrumentation = true;
 	}
 }
