Fixed UBT warning in TraceLog.Build.cs
- undid circular dependency between TraceLog and Core

#rnx

diff --git a/Engine/Source/Runtime/TraceLog/TraceLog.Build.cs b/Engine/Source/Runtime/TraceLog/TraceLog.Build.cs
--- a/Engine/Source/Runtime/TraceLog/TraceLog.Build.cs
+++ b/Engine/Source/Runtime/TraceLog/TraceLog.Build.cs
@@ -12,16 +12,11 @@ public class TraceLog : ModuleRules
 
 		bRequiresImplementModule = false;
 
-		if (Target.bBuildEditor && Target.bUseAutoRTFMCompiler)
-		{
-			// Surgical hack to temporarily allow Tracelog to build with AutoRTFM
-			// while the AutoRTFM team is working on a proper solution
-			PublicDependencyModuleNames.Add("Core");
-			CircularlyReferencedDependentModules.Add("Core");
-		}
-		else
+		PublicIncludePathModuleNames.Add("Core");
+
+		if (Target.bBuildEditor)
 		{
-			PublicIncludePathModuleNames.Add("Core");
+			bDisableAutoRTFMInstrumentation = true;
 		}
 
 		PrivateDefinitions.Add("SUPPRESS_PER_MODULE_INLINE_FILE"); // This module does not use core's standard operator new/delete overloads
