Fix deadlock on some kernels if there's many threads looking for fresh memory to trace to.

#rnx

diff --git a/Engine/Source/Runtime/TraceLog/Private/Trace/BlockPool.cpp b/Engine/Source/Runtime/TraceLog/Private/Trace/BlockPool.cpp
--- a/Engine/Source/Runtime/TraceLog/Private/Trace/BlockPool.cpp
+++ b/Engine/Source/Runtime/TraceLog/Private/Trace/BlockPool.cpp
@@ -5,6 +5,7 @@
 #if TRACE_PRIVATE_MINIMAL_ENABLED
 
 #include "Platform.h"
+#include "Misc/ScopeExit.h"
 #include "Trace/Detail/Atomic.h"
 #include "Trace/Detail/Writer.inl"
 
@@ -92,6 +93,9 @@ FWriteBuffer* Writer_AllocateBlockFromPool()
 	FWriteBuffer* Ret;
 	while (true)
 	{
+		int32 ThrottleRestore = ThreadThrottle();
+		ON_SCOPE_EXIT { ThreadUnthrottle(ThrottleRestore); };
+
 		// First we'll try one from the free list
 		FWriteBuffer* Owned = AtomicLoadRelaxed(&GPoolFreeList);
 		if (Owned != nullptr)
