[Trace] Fixed "-tracescopes" resulting in corrupt trace if scope functions emits sync trace events.

diff --git a/Engine/Source/Runtime/TraceLog/Private/Trace/Codec.cpp b/Engine/Source/Runtime/TraceLog/Private/Trace/Codec.cpp
--- a/Engine/Source/Runtime/TraceLog/Private/Trace/Codec.cpp
+++ b/Engine/Source/Runtime/TraceLog/Private/Trace/Codec.cpp
@@ -42,6 +42,18 @@ int32 Encode(const void* Src, int32 SrcSize, void* Dest, int32 DestSize)
 	);
 }
 
+////////////////////////////////////////////////////////////////////////////////
+int32 EncodeNoInstr(const void* Src, int32 SrcSize, void* Dest, int32 DestSize)
+{
+	return TRACE_PRIVATE_LZ4_NAMESPACE LZ4_compress_fast(
+		(const char*)Src,
+		(char*)Dest,
+		SrcSize,
+		DestSize,
+		1 // increase by 1 for small speed increase
+	);
+}
+
 ////////////////////////////////////////////////////////////////////////////////
 uint32 GetEncodeMaxSize(uint32 InputSize)
 {
diff --git a/Engine/Source/Runtime/TraceLog/Private/Trace/Writer.cpp b/Engine/Source/Runtime/TraceLog/Private/Trace/Writer.cpp
--- a/Engine/Source/Runtime/TraceLog/Private/Trace/Writer.cpp
+++ b/Engine/Source/Runtime/TraceLog/Private/Trace/Writer.cpp
@@ -37,6 +37,7 @@ namespace Private {
 
 ////////////////////////////////////////////////////////////////////////////////
 int32			Encode(const void*, int32, void*, int32);
+int32			EncodeNoInstr(const void*, int32, void*, int32);
 void			Writer_SetUpdateCallback(OnUpdateFunc* Callback);
 void			Writer_SendData(uint32, uint8* __restrict, uint32);
 void			Writer_InitializeTail(int32);
@@ -342,10 +343,8 @@ static bool Writer_FlushSendBuffer() { return true; }
 #endif
 
 ////////////////////////////////////////////////////////////////////////////////
-static void Writer_SendDataImpl(const void* Data, uint32 Size)
+static void Writer_SendDataImplNoInstr(const void* Data, uint32 Size)
 {
-	FProfilerScope _(__func__);
-
 #if TRACE_PRIVATE_STATISTICS
 	GTraceStatistics.BytesSent += Size;
 #endif
@@ -385,6 +384,13 @@ static void Writer_SendDataImpl(const void* Data, uint32 Size)
 #endif
 }
 
+////////////////////////////////////////////////////////////////////////////////
+static void Writer_SendDataImpl(const void* Data, uint32 Size)
+{
+	FProfilerScope _(__func__);
+	Writer_SendDataImplNoInstr(Data, Size);
+}
+
 ////////////////////////////////////////////////////////////////////////////////
 void Writer_SendDataRaw(const void* Data, uint32 Size)
 {
@@ -393,7 +399,8 @@ void Writer_SendDataRaw(const void* Data, uint32 Size)
 		return;
 	}
 
-	Writer_SendDataImpl(Data, Size);
+	FProfilerScope _(__func__);
+	Writer_SendDataImplNoInstr(Data, Size);
 }
 
 ////////////////////////////////////////////////////////////////////////////////
@@ -406,6 +413,59 @@ void Writer_SendDataNoInstr(uint32 ThreadId, uint8* __restrict Data, uint32 Size
 		return;
 	}
 
+	// Smaller buffers usually aren't redundant enough to benefit from being
+	// compressed. They often end up being larger.
+	if (Size <= 384)
+	{
+		Data -= sizeof(FTidPacket);
+		Size += sizeof(FTidPacket);
+		auto* Packet = (FTidPacket*)Data;
+		Packet->ThreadId = uint16(ThreadId & FTidPacketBase::ThreadIdMask);
+#if UE_TRACE_PACKET_VERIFICATION
+		Packet->ThreadId |= FTidPacketBase::Verification;
+#endif
+		Packet->PacketSize = uint16(Size);
+
+		Writer_SendDataImplNoInstr(Data, Size);
+
+#if UE_TRACE_PACKET_VERIFICATION
+		const uint64 Serial = GPacketSerial++;
+		Writer_SendDataImplNoInstr(&Serial, sizeof(uint64));
+#endif
+		return;
+	}
+
+	// Buffer size is expressed as "A + B" where A is a maximum expected
+	// input size (i.e. at least GPoolBlockSize) and B is LZ4 overhead as
+	// per LZ4_COMPRESSBOUND.
+	TTidPacketEncoded<8192 + 64> Packet;
+
+	Packet.ThreadId = FTidPacketBase::EncodedMarker;
+	Packet.ThreadId |= uint16(ThreadId & FTidPacketBase::ThreadIdMask);
+#if UE_TRACE_PACKET_VERIFICATION
+	Packet.ThreadId |= FTidPacketBase::Verification;
+#endif
+	Packet.DecodedSize = uint16(Size);
+	Packet.PacketSize = uint16(EncodeNoInstr(Data, Packet.DecodedSize, Packet.Data, sizeof(Packet.Data)));
+	Packet.PacketSize += sizeof(FTidPacketEncoded);
+
+	Writer_SendDataImplNoInstr(&Packet, Packet.PacketSize);
+
+#if UE_TRACE_PACKET_VERIFICATION
+	const uint64 Serial = GPacketSerial++;
+	Writer_SendDataImplNoInstr(&Serial, sizeof(uint64));
+#endif
+}
+
+////////////////////////////////////////////////////////////////////////////////
+void Writer_SendData(uint32 ThreadId, uint8* __restrict Data, uint32 Size)
+{
+	if (!GDataHandle)
+	{
+		return;
+	}
+
+	FProfilerScope _(__func__);
 
 	// Smaller buffers usually aren't redundant enough to benefit from being
 	// compressed. They often end up being larger.
@@ -444,6 +504,7 @@ void Writer_SendDataNoInstr(uint32 ThreadId, uint8* __restrict Data, uint32 Size
 	Packet.PacketSize += sizeof(FTidPacketEncoded);
 
 	Writer_SendDataImpl(&Packet, Packet.PacketSize);
+
 #if UE_TRACE_PACKET_VERIFICATION
 	const uint64 Serial = GPacketSerial++;
 	Writer_SendDataImpl(&Serial, sizeof(uint64));
@@ -451,15 +512,10 @@ void Writer_SendDataNoInstr(uint32 ThreadId, uint8* __restrict Data, uint32 Size
 }
 
 ////////////////////////////////////////////////////////////////////////////////
-void Writer_SendData(uint32 ThreadId, uint8* __restrict Data, uint32 Size)
+static void Writer_DescribeEvents(FEventNode::FIter Iter)
 {
 	FProfilerScope _(__func__);
-	Writer_SendDataNoInstr(ThreadId, Data, Size);
-}
 
-////////////////////////////////////////////////////////////////////////////////
-static void Writer_DescribeEvents(FEventNode::FIter Iter)
-{
 	// Since we are using on stack buffer to collect the describing events we
 	// cannot emit additional data on this thread. So use the un-instrumented
 	// send data functions here.
@@ -470,22 +526,38 @@ static void Writer_DescribeEvents(FEventNode::FIter Iter)
 		Event->Describe();
 
 		// Flush just in case an NewEvent event will be larger than 512 bytes.
-		if (TraceData.GetSize() >= (TraceData.GetCapacity() - 512))
+		uint32 RedirectSize = TraceData.GetSize();
+		if (RedirectSize >= (TraceData.GetCapacity() - 512))
 		{
-			Writer_SendDataNoInstr(ETransportTid::Events, TraceData.GetData(), TraceData.GetSize());
+			Writer_SendDataNoInstr(ETransportTid::Events, TraceData.GetData(), RedirectSize);
+			if (TraceData.GetSize() != RedirectSize)
+			{
+				// Trace events were emitted during Writer_SendDataNoInstr() call.
+				// These events will be lost!
+				PLATFORM_BREAK();
+				UE_TRACE_MESSAGE(WriterError, "Events emitted while describing events; trace will be corrupt!");
+			}
 			TraceData.Reset();
 		}
 	}
 
-	if (TraceData.GetSize())
+	if (uint32 RedirectSize = TraceData.GetSize())
 	{
-		Writer_SendDataNoInstr(ETransportTid::Events, TraceData.GetData(), TraceData.GetSize());
+		Writer_SendDataNoInstr(ETransportTid::Events, TraceData.GetData(), RedirectSize);
+		if (TraceData.GetSize() != RedirectSize)
+		{
+			// Trace events were emitted during Writer_SendDataNoInstr() call.
+			// These events will be lost!
+			PLATFORM_BREAK();
+			UE_TRACE_MESSAGE(WriterError, "Events emitted while describing events; trace will be corrupt!");
+		}
 	}
 }
 
 ////////////////////////////////////////////////////////////////////////////////
 static void Writer_AnnounceChannels()
 {
+	FProfilerScope _(__func__);
 	FChannel::Iter Iter = FChannel::ReadNew();
 	while (const FChannel* Channel = Iter.GetNext())
 	{
