Fix FImportantLogScope under transactions.

diff --git a/Engine/Source/Programs/AutoRTFMTests/Private/InsightsTests.cpp b/Engine/Source/Programs/AutoRTFMTests/Private/InsightsTests.cpp
--- a/Engine/Source/Programs/AutoRTFMTests/Private/InsightsTests.cpp
+++ b/Engine/Source/Programs/AutoRTFMTests/Private/InsightsTests.cpp
@@ -28,7 +28,15 @@ private:
 	const bool bWasEnabled;
 };
 
-UE_TRACE_EVENT_BEGIN(Cpu, SomeTraceEvent, NoSync)
+UE_TRACE_EVENT_BEGIN(Cpu, SomeTraceEvent)
+	UE_TRACE_EVENT_FIELD(int32, Foo)
+UE_TRACE_EVENT_END()
+
+UE_TRACE_EVENT_BEGIN(Cpu, SomeNoSyncTraceEvent, NoSync)
+	UE_TRACE_EVENT_FIELD(int32, Foo)
+UE_TRACE_EVENT_END()
+
+UE_TRACE_EVENT_BEGIN(Cpu, SomeImportantTraceEvent, NoSync | Important)
 	UE_TRACE_EVENT_FIELD(int32, Foo)
 UE_TRACE_EVENT_END()
 
@@ -479,11 +487,53 @@ TEST_CASE("Insights")
 	}
 	
 #if CPUPROFILERTRACE_ENABLED
+	// Some number over 10000 to force allocations.
+	constexpr unsigned Iterations = 16384;
 	SECTION("UE_TRACE_LOG_SCOPED_T")
 	{
-		AutoRTFM::Testing::Commit([&]
+		AutoRTFM::Testing::Abort([&]
 		{
 			UE_TRACE_LOG_SCOPED_T(Cpu, SomeTraceEvent, CpuChannel) << SomeTraceEvent.Foo(42);
+			AutoRTFM::AbortTransaction();
+		});
+
+		AutoRTFM::Testing::Commit([&]
+		{
+			for (int Index = 0; Index < Iterations; Index++)
+			{
+				UE_TRACE_LOG_SCOPED_T(Cpu, SomeTraceEvent, CpuChannel) << SomeTraceEvent.Foo(42);
+			}
+		});
+
+		AutoRTFM::Testing::Abort([&]
+		{
+			UE_TRACE_LOG_SCOPED_T(Cpu, SomeNoSyncTraceEvent, CpuChannel) << SomeNoSyncTraceEvent.Foo(42);
+			AutoRTFM::AbortTransaction();
+		});
+
+		AutoRTFM::Testing::Commit([&]
+		{
+			for (int Index = 0; Index < Iterations; Index++)
+			{
+				UE_TRACE_LOG_SCOPED_T(Cpu, SomeNoSyncTraceEvent, CpuChannel) << SomeNoSyncTraceEvent.Foo(42);
+			}
+		});
+	}
+
+	SECTION("UE_TRACE_LOG")
+	{
+		AutoRTFM::Testing::Abort([&]
+		{
+			UE_TRACE_LOG(Cpu, SomeImportantTraceEvent, CpuChannel) << SomeImportantTraceEvent.Foo(42);
+			AutoRTFM::AbortTransaction();
+		});
+
+		AutoRTFM::Testing::Commit([&]
+		{
+			for (int Index = 0; Index < Iterations; Index++)
+			{
+				UE_TRACE_LOG(Cpu, SomeImportantTraceEvent, CpuChannel) << SomeImportantTraceEvent.Foo(42);
+			}
 		});
 	}
 #endif
diff --git a/Engine/Source/Runtime/TraceLog/Public/Trace/Detail/Important/ImportantLogScope.inl b/Engine/Source/Runtime/TraceLog/Public/Trace/Detail/Important/ImportantLogScope.inl
--- a/Engine/Source/Runtime/TraceLog/Public/Trace/Detail/Important/ImportantLogScope.inl
+++ b/Engine/Source/Runtime/TraceLog/Public/Trace/Detail/Important/ImportantLogScope.inl
@@ -9,6 +9,7 @@
 #include "Trace/Detail/Writer.inl"
 #include "Trace/Detail/Field.h"
 #include "Trace/Detail/EventNode.h"
+#include "AutoRTFM/AutoRTFM.h"
 
 namespace UE {
 namespace Trace {
@@ -24,7 +25,7 @@ TRACELOG_API FNextSharedBuffer				Writer_NextSharedBuffer(FSharedBuffer*, int32,
 
 ////////////////////////////////////////////////////////////////////////////////
 template <class T>
-FORCENOINLINE FImportantLogScope FImportantLogScope::Enter(uint32 ArrayDataSize)
+UE_AUTORTFM_ALWAYS_OPEN FORCENOINLINE FImportantLogScope FImportantLogScope::Enter(uint32 ArrayDataSize)
 {
 	static_assert(!!(uint32(T::EventFlags) & uint32(FEventInfo::Flag_MaybeHasAux)), "Only important trace events with array-type fields need a size parameter to UE_TRACE_LOG()");
 
@@ -42,7 +43,7 @@ FORCENOINLINE FImportantLogScope FImportantLogScope::Enter(uint32 ArrayDataSize)
 
 ////////////////////////////////////////////////////////////////////////////////
 template <class T>
-inline FImportantLogScope FImportantLogScope::Enter()
+UE_AUTORTFM_ALWAYS_OPEN inline FImportantLogScope FImportantLogScope::Enter()
 {
 	static_assert(!(uint32(T::EventFlags) & uint32(FEventInfo::Flag_MaybeHasAux)), "Important trace events with array-type fields must be traced with UE_TRACE_LOG(Logger, Event, Channel, ArrayDataSize)");
 
@@ -52,7 +53,7 @@ inline FImportantLogScope FImportantLogScope::Enter()
 }
 
 ////////////////////////////////////////////////////////////////////////////////
-inline FImportantLogScope FImportantLogScope::EnterImpl(uint32 Uid, uint32 Size)
+UE_AUTORTFM_ALWAYS_OPEN inline FImportantLogScope FImportantLogScope::EnterImpl(uint32 Uid, uint32 Size)
 {
 	FSharedBuffer* Buffer = AtomicLoadAcquire(&GSharedBuffer);
 
@@ -85,18 +86,17 @@ inline FImportantLogScope FImportantLogScope::EnterImpl(uint32 Uid, uint32 Size)
 }
 
 ////////////////////////////////////////////////////////////////////////////////
-inline void FImportantLogScope::operator += (const FImportantLogScope&) const
+UE_AUTORTFM_ALWAYS_OPEN inline void FImportantLogScope::operator += (const FImportantLogScope& Other) const
 {
 	auto* Buffer = (FSharedBuffer*)(Ptr + BufferOffset);
 	AtomicAddRelease(&(Buffer->Cursor), int32(FSharedBuffer::RefBit));
 }
 
-
-
 ////////////////////////////////////////////////////////////////////////////////
 template <typename FieldMeta, typename Type>
 struct FImportantLogScope::FFieldSet
 {
+	UE_AUTORTFM_ALWAYS_OPEN
 	static void Impl(FImportantLogScope* Scope, const Type& Value)
 	{
 		uint8* Dest = (uint8*)(Scope->Ptr) + FieldMeta::Offset;
@@ -108,6 +108,7 @@ struct FImportantLogScope::FFieldSet
 template <typename FieldMeta, typename Type>
 struct FImportantLogScope::FFieldSet<FieldMeta, Type[]>
 {
+	UE_AUTORTFM_ALWAYS_OPEN
 	static void Impl(FImportantLogScope* Scope, Type const* Data, int32 Num)
 	{
 		uint32 Size = Num * sizeof(Type);
@@ -130,6 +131,7 @@ struct FImportantLogScope::FFieldSet<FieldMeta, Type[]>
 template <typename FieldMeta>
 struct FImportantLogScope::FFieldSet<FieldMeta, AnsiString>
 {
+	UE_AUTORTFM_ALWAYS_OPEN
 	static void Impl(FImportantLogScope* Scope, const ANSICHAR* String, int32 Length=-1)
 	{
 		if (Length < 0)
@@ -150,6 +152,7 @@ struct FImportantLogScope::FFieldSet<FieldMeta, AnsiString>
 		Scope->Ptr[Scope->AuxCursor] = uint8(EKnownEventUids::AuxDataTerminal);
 	}
 
+	UE_AUTORTFM_ALWAYS_OPEN
 	static void Impl(FImportantLogScope* Scope, const WIDECHAR* String, int32 Length=-1)
 	{
 		if (Length < 0)
@@ -182,6 +185,7 @@ struct FImportantLogScope::FFieldSet<FieldMeta, AnsiString>
 template <typename FieldMeta>
 struct FImportantLogScope::FFieldSet<FieldMeta, WideString>
 {
+	UE_AUTORTFM_ALWAYS_OPEN
 	static void Impl(FImportantLogScope* Scope, const WIDECHAR* String, int32 Length=-1)
 	{
 		if (Length < 0)
@@ -210,6 +214,7 @@ struct FImportantLogScope::FFieldSet<FieldMeta, WideString>
 template <typename FieldMeta, typename DefinitionType>
 struct FImportantLogScope::FFieldSet<FieldMeta, TEventRef<DefinitionType>>
 {
+	UE_AUTORTFM_ALWAYS_OPEN
 	static void Impl(FImportantLogScope* Scope, const TEventRef<DefinitionType>& Reference)
 	{
 		FFieldSet<FieldMeta, DefinitionType>::Impl(Scope, Reference.Id);
diff --git a/Engine/Source/Runtime/TraceLog/Public/Trace/Detail/LogScope.h b/Engine/Source/Runtime/TraceLog/Public/Trace/Detail/LogScope.h
--- a/Engine/Source/Runtime/TraceLog/Public/Trace/Detail/LogScope.h
+++ b/Engine/Source/Runtime/TraceLog/Public/Trace/Detail/LogScope.h
@@ -65,6 +65,9 @@ public:
 			~FScopedLogScope();
 	void	SetActive();
 	bool	bActive = false;
+
+private:
+	void    Deinit();
 };
 
 ////////////////////////////////////////////////////////////////////////////////
@@ -74,6 +77,9 @@ public:
 			~FScopedStampedLogScope();
 	void	SetActive();
 	bool	bActive = false;
+
+private:
+	void    Deinit();
 };
 
 } // namespace Private
diff --git a/Engine/Source/Runtime/TraceLog/Public/Trace/Detail/LogScope.inl b/Engine/Source/Runtime/TraceLog/Public/Trace/Detail/LogScope.inl
--- a/Engine/Source/Runtime/TraceLog/Public/Trace/Detail/LogScope.inl
+++ b/Engine/Source/Runtime/TraceLog/Public/Trace/Detail/LogScope.inl
@@ -12,6 +12,7 @@
 #include "LogScope.h"
 #include "Protocol.h"
 #include "Writer.inl"
+#include "AutoRTFM/AutoRTFM.h"
 
 namespace UE {
 namespace Trace {
@@ -22,12 +23,14 @@ extern TRACELOG_API uint64			GStartCycle;
 extern TRACELOG_API uint32 volatile	GLogSerial;
 
 ////////////////////////////////////////////////////////////////////////////////
+UE_AUTORTFM_ALWAYS_OPEN
 inline void FLogScope::Commit() const
 {
 	AtomicStoreRelease((uint8**) &(Buffer->Committed), Buffer->Cursor);
 }
 
 ////////////////////////////////////////////////////////////////////////////////
+UE_AUTORTFM_ALWAYS_OPEN
 inline void FLogScope::Commit(FWriteBuffer* __restrict LatestBuffer) const
 {
 	if (LatestBuffer != Buffer)
@@ -40,6 +43,7 @@ inline void FLogScope::Commit(FWriteBuffer* __restrict LatestBuffer) const
 
 ////////////////////////////////////////////////////////////////////////////////
 template <uint32 Flags>
+UE_AUTORTFM_ALWAYS_OPEN
 inline auto FLogScope::EnterImpl(uint32 Uid, uint32 Size)
 {
 	TLogScope<(Flags & FEventInfo::Flag_MaybeHasAux) != 0> Ret;
@@ -56,6 +60,7 @@ inline auto FLogScope::EnterImpl(uint32 Uid, uint32 Size)
 
 ////////////////////////////////////////////////////////////////////////////////
 template <class HeaderType>
+UE_AUTORTFM_ALWAYS_OPEN
 inline void FLogScope::EnterPrelude(uint32 Size)
 {
 	uint32 AllocSize = sizeof(HeaderType) + Size;
@@ -81,6 +86,7 @@ inline void FLogScope::EnterPrelude(uint32 Size)
 }
 
 ////////////////////////////////////////////////////////////////////////////////
+UE_AUTORTFM_ALWAYS_OPEN
 inline void FLogScope::Enter(uint32 Uid, uint32 Size)
 {
 	EnterPrelude<FEventHeaderSync>(Size);
@@ -94,6 +100,7 @@ inline void FLogScope::Enter(uint32 Uid, uint32 Size)
 }
 
 ////////////////////////////////////////////////////////////////////////////////
+UE_AUTORTFM_ALWAYS_OPEN
 inline void FLogScope::EnterNoSync(uint32 Uid, uint32 Size)
 {
 	EnterPrelude<FEventHeader>(Size);
@@ -107,6 +114,7 @@ inline void FLogScope::EnterNoSync(uint32 Uid, uint32 Size)
 
 ////////////////////////////////////////////////////////////////////////////////
 template <bool bMaybeHasAux>
+UE_AUTORTFM_ALWAYS_OPEN
 inline void TLogScope<bMaybeHasAux>::operator += (const FLogScope&) const
 {
 	if constexpr (bMaybeHasAux)
@@ -124,13 +132,23 @@ inline void TLogScope<bMaybeHasAux>::operator += (const FLogScope&) const
 }
 
 ////////////////////////////////////////////////////////////////////////////////
-inline FScopedLogScope::~FScopedLogScope()
+inline void FScopedLogScope::Deinit()
 {
 	if (!bActive)
 	{
 		return;
 	}
 
+	if (AutoRTFM::IsClosed())
+	{
+		// If we are operating closed (in a transaction) we have to forget about
+		// our previously registered on-abort handler, because we're currently
+		// about to close out the scope within the transaction.
+		AutoRTFM::PopOnAbortHandler(this);
+		AutoRTFM::Open([this] { this->Deinit(); });
+		return;
+	}
+
 	uint8 LeaveUid = uint8(EKnownEventUids::LeaveScope << EKnownEventUids::_UidShift);
 
 	FWriteBuffer* Buffer = Writer_GetBuffer();
@@ -146,21 +164,51 @@ inline FScopedLogScope::~FScopedLogScope()
 }
 
 ////////////////////////////////////////////////////////////////////////////////
-inline void FScopedLogScope::SetActive()
+inline FScopedLogScope::~FScopedLogScope()
 {
-	bActive = true;
+	Deinit();
 }
 
+////////////////////////////////////////////////////////////////////////////////
+inline void FScopedLogScope::SetActive()
+{
+	if (bActive)
+	{
+		// We don't want to re-activate an already active scope!
+		return;
+	}
+
+	if (AutoRTFM::IsClosed())
+	{
+		// If we are operating closed (in a transaction) we could abort before
+		// `this` hits its destructor. In that case we need to remember when
+		// we abort to close out the scoped trace event.
+		AutoRTFM::PushOnAbortHandler(this, [this] { this->Deinit(); });
+		AutoRTFM::Open([this] { this->SetActive(); });
+		return;
+	}
 
+	bActive = true;
+}
 
 ////////////////////////////////////////////////////////////////////////////////
-inline FScopedStampedLogScope::~FScopedStampedLogScope()
+inline void FScopedStampedLogScope::Deinit()
 {
 	if (!bActive)
 	{
 		return;
 	}
 
+	if (AutoRTFM::IsClosed())
+	{
+		// If we are operating closed (in a transaction) we have to forget about
+		// our previously registered on-abort handler, because we're currently
+		// about to close out the scope within the transaction.
+		AutoRTFM::PopOnAbortHandler(this);
+		AutoRTFM::Open([this] { this->Deinit(); });
+		return;
+	}
+
 	uint64 Stamp = TimeGetTimestamp() - GStartCycle;
 
 	FWriteBuffer* Buffer = Writer_GetBuffer();
@@ -178,15 +226,36 @@ inline FScopedStampedLogScope::~FScopedStampedLogScope()
 }
 
 ////////////////////////////////////////////////////////////////////////////////
-inline void FScopedStampedLogScope::SetActive()
+inline FScopedStampedLogScope::~FScopedStampedLogScope()
 {
-	bActive = true;
+	Deinit();
 }
 
+////////////////////////////////////////////////////////////////////////////////
+inline void FScopedStampedLogScope::SetActive()
+{
+	if (bActive)
+	{
+		// We don't want to re-activate an already active scope!
+		return;
+	}
+
+	if (AutoRTFM::IsClosed())
+	{
+		// If we are operating closed (in a transaction) we could abort before
+		// `this` hits its destructor. In that case we need to remember when
+		// we abort to close out the scoped trace event.
+		AutoRTFM::PushOnAbortHandler(this, [this] { this->Deinit(); });
+		AutoRTFM::Open([this] { this->SetActive(); });
+		return;
+	}
 
+	bActive = true;
+}
 
 ////////////////////////////////////////////////////////////////////////////////
 template <class EventType>
+UE_AUTORTFM_ALWAYS_OPEN
 FORCENOINLINE auto FLogScope::Enter()
 {
 	uint32 Size = EventType::GetSize();
@@ -196,6 +265,7 @@ FORCENOINLINE auto FLogScope::Enter()
 
 ////////////////////////////////////////////////////////////////////////////////
 template <class EventType>
+UE_AUTORTFM_ALWAYS_OPEN
 FORCENOINLINE auto FLogScope::ScopedEnter()
 {
 	uint8 EnterUid = uint8(EKnownEventUids::EnterScope << EKnownEventUids::_UidShift);
@@ -218,6 +288,7 @@ FORCENOINLINE auto FLogScope::ScopedEnter()
 
 ////////////////////////////////////////////////////////////////////////////////
 template <class EventType>
+UE_AUTORTFM_ALWAYS_OPEN
 FORCENOINLINE auto FLogScope::ScopedStampedEnter()
 {
 	uint64 Stamp;
@@ -247,6 +318,7 @@ FORCENOINLINE auto FLogScope::ScopedStampedEnter()
 template <typename FieldMeta, typename Type>
 struct FLogScope::FFieldSet
 {
+	UE_AUTORTFM_ALWAYS_OPEN
 	static void Impl(FLogScope* Scope, const Type& Value)
 	{
 		uint8* Dest = (uint8*)(Scope->Ptr) + FieldMeta::Offset;
@@ -258,6 +330,7 @@ struct FLogScope::FFieldSet
 template <typename FieldMeta, typename Type>
 struct FLogScope::FFieldSet<FieldMeta, Type[]>
 {
+	UE_AUTORTFM_ALWAYS_OPEN
 	static void Impl(FLogScope*, Type const* Data, int32 Num)
 	{
 		static const uint32 Index = FieldMeta::Index & int32(EIndexPack::NumFieldsMask);
@@ -279,6 +352,7 @@ struct FLogScope::FFieldSet<FieldMeta, Type[Count]>
 template <typename FieldMeta>
 struct FLogScope::FFieldSet<FieldMeta, AnsiString>
 {
+	UE_AUTORTFM_ALWAYS_OPEN
 	static void Impl(FLogScope*, const ANSICHAR* String, int32 Length=-1)
 	{
 		if (Length < 0)
@@ -289,7 +363,8 @@ struct FLogScope::FFieldSet<FieldMeta, AnsiString>
 		static const uint32 Index = FieldMeta::Index & int32(EIndexPack::NumFieldsMask);
 		Field_WriteStringAnsi(Index, String, Length);
 	}
-
+	
+	UE_AUTORTFM_ALWAYS_OPEN
 	static void Impl(FLogScope*, const WIDECHAR* String, int32 Length=-1)
 	{
 		if (Length < 0)
@@ -307,6 +382,7 @@ struct FLogScope::FFieldSet<FieldMeta, AnsiString>
 template <typename FieldMeta>
 struct FLogScope::FFieldSet<FieldMeta, WideString>
 {
+	UE_AUTORTFM_ALWAYS_OPEN
 	static void Impl(FLogScope*, const WIDECHAR* String, int32 Length=-1)
 	{
 		if (Length < 0)
@@ -324,6 +400,7 @@ struct FLogScope::FFieldSet<FieldMeta, WideString>
 template <typename FieldMeta, typename DefinitionType>
 struct FLogScope::FFieldSet<FieldMeta, TEventRef<DefinitionType>>
 {
+	UE_AUTORTFM_ALWAYS_OPEN
 	static void Impl(FLogScope* Scope, const TEventRef<DefinitionType>& Reference)
 	{
 		FFieldSet<FieldMeta, DefinitionType>::Impl(Scope, Reference.Id);
