[Trace]
- Fix TSAN warning while grabbing a copy of statistics

diff --git a/Engine/Source/Runtime/TraceLog/Private/Trace/Trace.cpp b/Engine/Source/Runtime/TraceLog/Private/Trace/Trace.cpp
--- a/Engine/Source/Runtime/TraceLog/Private/Trace/Trace.cpp
+++ b/Engine/Source/Runtime/TraceLog/Private/Trace/Trace.cpp
@@ -101,7 +101,16 @@ void Update()
 ////////////////////////////////////////////////////////////////////////////////
 void GetStatistics(FStatistics& Out)
 {
-	Out = Private::GTraceStatistics;
+	Out.BytesSent = Private::AtomicLoadRelaxed(&Private::GTraceStatistics.BytesSent);
+	Out.BytesTraced = Private::AtomicLoadRelaxed(&Private::GTraceStatistics.BytesTraced);
+	Out.BytesEmitted = Private::AtomicLoadRelaxed(&Private::GTraceStatistics.BytesEmitted);
+	Out.MemoryUsed = Private::AtomicLoadRelaxed(&Private::GTraceStatistics.MemoryUsed);
+	Out.BlockPoolAllocated = Private::AtomicLoadRelaxed(&Private::GTraceStatistics.BlockPoolAllocated);
+	Out.SharedBufferAllocated = Private::AtomicLoadRelaxed(&Private::GTraceStatistics.SharedBufferAllocated);
+	Out.FixedBufferAllocated = Private::AtomicLoadRelaxed(&Private::GTraceStatistics.FixedBufferAllocated);
+	Out.CacheAllocated = Private::AtomicLoadRelaxed(&Private::GTraceStatistics.CacheAllocated);
+	Out.CacheUsed = Private::AtomicLoadRelaxed(&Private::GTraceStatistics.CacheUsed);
+	Out.CacheWaste = Private::AtomicLoadRelaxed(&Private::GTraceStatistics.CacheWaste);
 }
 
 ////////////////////////////////////////////////////////////////////////////////
