TraceLog limits on block pool size.

Adds a hard limit on block pool growth. Once the limit is hit and no free buffers are available, tracing threads will start to block until the worker thread has had time to write out data and give back buffers to the pool.

This also solves the issue with sending out serial events in incorrect order during periods of extreme load. This is because we can calculate a worst case limit that guarantees two events with the same serial are in memory waiting to be written to the trace stream. Since the smallest possible serial event is 5 bytes and serial wrap every 16M event gives a theoretical limit of 79 MiB (not accounting for buffer tracking overhead). The limit must not be enforced while the worker thread is not running (e.g. during process initialization and shutdown).

diff --git a/Engine/Source/Runtime/Launch/Private/LaunchEngineLoop.cpp b/Engine/Source/Runtime/Launch/Private/LaunchEngineLoop.cpp
--- a/Engine/Source/Runtime/Launch/Private/LaunchEngineLoop.cpp
+++ b/Engine/Source/Runtime/Launch/Private/LaunchEngineLoop.cpp
@@ -5282,6 +5282,8 @@ void FEngineLoop::Exit()
 #if WITH_ENGINE && FRAMEPRO_ENABLED
 	FFrameProProfiler::TearDown();
 #endif // FRAMEPRO_ENABLED
+
+	UE::Trace::Exit();
 }
 
 
diff --git a/Engine/Source/Runtime/TraceLog/Private/Trace/BlockPool.cpp b/Engine/Source/Runtime/TraceLog/Private/Trace/BlockPool.cpp
--- a/Engine/Source/Runtime/TraceLog/Private/Trace/BlockPool.cpp
+++ b/Engine/Source/Runtime/TraceLog/Private/Trace/BlockPool.cpp
@@ -16,6 +16,7 @@ namespace Private {
 ////////////////////////////////////////////////////////////////////////////////
 void*		Writer_MemoryAllocate(SIZE_T, uint32);
 void		Writer_MemoryFree(void*, uint32);
+void		LogStall(uint64, uint64);
 
 extern FStatistics GTraceStatistics;
 
@@ -40,11 +41,31 @@ struct FPoolBlockList
 static const uint32						GPoolBlockSize		= 4 << 10;
 static const uint32						GPoolPageSize		= GPoolBlockSize << 4;
 static const uint32						GPoolInitPageSize	= GPoolBlockSize << 6;
+T_ALIGN static volatile uint32			GPoolTotalAllocated;// = 0
+static uint32							GPoolMaxSize		= ~0; // Starts unlimited, written by initialization
 T_ALIGN static FWriteBuffer* volatile	GPoolFreeList;		// = nullptr;
 T_ALIGN static UPTRINT volatile			GPoolFutex;			// = 0
 T_ALIGN static FPoolPage* volatile		GPoolPageList;		// = nullptr;
 #undef T_ALIGN
 
+////////////////////////////////////////////////////////////////////////////////
+void Writer_SetBlockPoolLimit(uint32 MaxSize)
+{
+	/*
+	 * Note that setting it higher than the default value introduces risks that
+	 * serialized events will be transmitted in the wrong order. The smallest 
+	 * possible size for a (serialized) event, one with no field, is 5 bytes.
+	 * Serials wrap every 16M event. This gives a theoretical limit of 79MiB.
+	 */
+	GPoolMaxSize = MaxSize; 
+}
+
+////////////////////////////////////////////////////////////////////////////////
+void Writer_UnsetBlockPoolLimit()
+{
+	GPoolMaxSize = ~0;
+}
+
 ////////////////////////////////////////////////////////////////////////////////
 static FPoolBlockList Writer_AllocateBlockList(uint32 PageSize)
 {
@@ -54,6 +75,7 @@ static FPoolBlockList Writer_AllocateBlockList(uint32 PageSize)
 #if TRACE_PRIVATE_STATISTICS
 	AtomicAddRelaxed(&GTraceStatistics.BlockPoolAllocated, uint64(PageSize));
 #endif
+	AtomicAddRelaxed(&GPoolTotalAllocated, PageSize);
 
 	uint32 BufferSize = GPoolBlockSize;
 	BufferSize -= sizeof(FWriteBuffer);
@@ -90,6 +112,7 @@ static FPoolBlockList Writer_AllocateBlockList(uint32 PageSize)
 FWriteBuffer* Writer_AllocateBlockFromPool()
 {
 	// Fetch a new buffer
+	uint64 StallStart = 0;
 	FWriteBuffer* Ret;
 	while (true)
 	{
@@ -116,6 +139,14 @@ FWriteBuffer* Writer_AllocateBlockFromPool()
 			break;
 		}
 
+		// If we have hit the pool limit, continue looping until the worker
+		// thread has had time write out and give back free buffers.
+		if (AtomicLoadAcquire(&GPoolTotalAllocated) >= GPoolMaxSize)
+		{
+			StallStart = StallStart ? StallStart : TimeGetRelativeTimestamp();
+			continue;
+		}
+
 		// The free list is empty. Map some more memory.
 		UPTRINT Futex = AtomicLoadRelaxed(&GPoolFutex);
 		if (Futex || !AtomicCompareExchangeAcquire(&GPoolFutex, Futex + 1, Futex))
@@ -153,6 +184,12 @@ FWriteBuffer* Writer_AllocateBlockFromPool()
 		break;
 	}
 
+	if (StallStart)
+	{
+		const uint64 StallEnd = TimeGetRelativeTimestamp();
+		LogStall(StallStart, StallEnd);
+	}
+
 	return Ret;
 }
 
@@ -186,10 +223,11 @@ void Writer_ShutdownPool()
 	for (auto* Page = AtomicLoadRelaxed(&GPoolPageList); Page != nullptr;)
 	{
 		FPoolPage* NextPage = Page->NextPage;
-		Writer_MemoryFree(Page, Page->AllocSize);
 #if TRACE_PRIVATE_STATISTICS
 		AtomicSubRelease(&GTraceStatistics.BlockPoolAllocated, uint64(Page->AllocSize));
 #endif
+		AtomicSubRelease(&GPoolTotalAllocated, Page->AllocSize);
+		Writer_MemoryFree(Page, Page->AllocSize);
 		Page = NextPage;
 	}
 }
diff --git a/Engine/Source/Runtime/TraceLog/Private/Trace/Platform.h b/Engine/Source/Runtime/TraceLog/Private/Trace/Platform.h
--- a/Engine/Source/Runtime/TraceLog/Private/Trace/Platform.h
+++ b/Engine/Source/Runtime/TraceLog/Private/Trace/Platform.h
@@ -19,7 +19,17 @@ void	ThreadDestroy(UPTRINT Handle);
 
 ////////////////////////////////////////////////////////////////////////////////
 uint64				TimeGetFrequency();
+/**
+ * Gets the absolute timestamp of the system.
+ * @return Absolute timestamp
+ */
 TRACELOG_API uint64	TimeGetTimestamp();
+/**
+ * Gets the timestamp (relative to base). If called
+ * before Trace initialization value will be zero.
+ * @return Timestamp relative base time stamp
+ */
+TRACELOG_API uint64 TimeGetRelativeTimestamp();
 
 ////////////////////////////////////////////////////////////////////////////////
 UPTRINT	TcpSocketConnect(const ANSICHAR* Host, uint16 Port);
diff --git a/Engine/Source/Runtime/TraceLog/Private/Trace/Trace.cpp b/Engine/Source/Runtime/TraceLog/Private/Trace/Trace.cpp
--- a/Engine/Source/Runtime/TraceLog/Private/Trace/Trace.cpp
+++ b/Engine/Source/Runtime/TraceLog/Private/Trace/Trace.cpp
@@ -20,6 +20,7 @@ void	Writer_SetUpdateCallback(OnUpdateFunc* Callback);
 void	Writer_MemorySetHooks(AllocFunc, FreeFunc);
 void	Writer_Initialize(const FInitializeDesc&);
 void	Writer_WorkerCreate();
+void	Writer_UnsetBlockPoolLimit();
 void	Writer_Shutdown();
 void	Writer_Update();
 bool	Writer_SendTo(const ANSICHAR*, uint32, uint32);
@@ -34,6 +35,18 @@ uint32	Writer_GetThreadId();
 
 extern FStatistics GTraceStatistics;
 
+UE_TRACE_MINIMAL_EVENT_BEGIN($Trace, TraceStall)
+	UE_TRACE_MINIMAL_EVENT_FIELD(uint64, StartCycle)
+	UE_TRACE_MINIMAL_EVENT_FIELD(uint64, EndCycle)
+UE_TRACE_MINIMAL_EVENT_END()
+
+void LogStall(uint64 Start, uint64 End)
+{
+	UE_TRACE_LOG($Trace, TraceStall, true)
+		<< TraceStall.StartCycle(Start)
+		<< TraceStall.EndCycle(End);
+}
+
 } // namespace Private
 
 
@@ -80,6 +93,12 @@ void Initialize(const FInitializeDesc& Desc)
 	FChannel::Initialize();
 }
 
+////////////////////////////////////////////////////////////////////////////////
+void Exit()
+{
+	Private::Writer_UnsetBlockPoolLimit();
+}
+
 ////////////////////////////////////////////////////////////////////////////////
 void Shutdown()
 {
diff --git a/Engine/Source/Runtime/TraceLog/Private/Trace/Writer.cpp b/Engine/Source/Runtime/TraceLog/Private/Trace/Writer.cpp
--- a/Engine/Source/Runtime/TraceLog/Private/Trace/Writer.cpp
+++ b/Engine/Source/Runtime/TraceLog/Private/Trace/Writer.cpp
@@ -62,6 +62,8 @@ void			Writer_ShutdownControl();
 bool			Writer_IsTailing();
 static bool		Writer_SessionPrologue();
 void			Writer_FreeBlockListToPool(FWriteBuffer*, FWriteBuffer*);
+void			Writer_SetBlockPoolLimit(uint32);
+void			Writer_UnsetBlockPoolLimit();
 
 
 ////////////////////////////////////////////////////////////////////////////////
@@ -85,6 +87,7 @@ TRACELOG_API uint64				GStartCycle;		// = 0;
 TRACELOG_API uint32 volatile	GLogSerial;			// = 0;
 // Counter of calls to Writer_WorkerUpdate to enable regular flushing of output buffers
 static uint32					GUpdateCounter;		// = 0;
+static uint32					GBlockPoolMaxSize;
 #if UE_TRACE_PACKET_VERIFICATION
 uint64							GPacketSerial = 1;
 #endif
@@ -177,6 +180,11 @@ uint32 Writer_GetThreadId()
 
 #endif // UE_TRACE_USE_TLS_CONTEXT_OBJECT
 
+////////////////////////////////////////////////////////////////////////////////
+uint64 TimeGetRelativeTimestamp()
+{
+	return GStartCycle ? TimeGetTimestamp() - GStartCycle : 0;
+}
 
 ////////////////////////////////////////////////////////////////////////////////
 void Writer_CreateGuid(FTraceGuid* OutGuid)
@@ -843,12 +851,18 @@ static void Writer_WorkerThread()
 {
 	ThreadRegister(TEXT("Trace"), 0, INT_MAX);
 
+	// Enable the pool limit when the worker thread is running
+	Writer_SetBlockPoolLimit(GBlockPoolMaxSize);
+
 	while (!AtomicLoadRelaxed(&GWorkerThreadQuit))
 	{
 		Writer_WorkerUpdate();
 
 		ThreadSleep(GSleepTimeInMS);
 	}
+
+	// Reset the limit as no one will pick up data
+	Writer_UnsetBlockPoolLimit();
 }
 
 ////////////////////////////////////////////////////////////////////////////////
@@ -979,6 +993,7 @@ void Writer_Initialize(const FInitializeDesc& Desc)
 	FProfilerScope::OnScopeEnd = Desc.OnScopeEndFunc;
 
 	FProfilerScope _(__func__);
+	GBlockPoolMaxSize = Desc.BlockPoolMaxSize;
 
 	Writer_InitializeTail(Desc.TailSizeBytes);
 
diff --git a/Engine/Source/Runtime/TraceLog/Public/Trace/Config.h b/Engine/Source/Runtime/TraceLog/Public/Trace/Config.h
--- a/Engine/Source/Runtime/TraceLog/Public/Trace/Config.h
+++ b/Engine/Source/Runtime/TraceLog/Public/Trace/Config.h
@@ -113,10 +113,18 @@ static_assert(UE_TRACE_MINIMAL_ENABLED == 0 || UE_TRACE_ENABLED_SHIPPING_EXPERIM
 #	define TRACE_PRIVATE_ALLOW_IMPORTANTS UE_TRACE_ALLOW_IMPORTANT_EVENTS
 #endif
 
+/**
+ * Default block pool size in bytes. See BlockPool implementation for discussion about
+ * overriding this value.
+ */
+#if !defined(UE_TRACE_BLOCK_POOL_MAXSIZE)
+#	define UE_TRACE_BLOCK_POOL_MAXSIZE 79 << 20
+#endif
+
 /**
  * Enable packet verification. Only useful when looking for transmission bugs. Note that in order
  * to avoid making a new protocol version, enabling this makes existing version 7 traces incompatible.
  */
 #ifndef UE_TRACE_PACKET_VERIFICATION
 #	define UE_TRACE_PACKET_VERIFICATION 0
-#endif
\ No newline at end of file
+#endif
diff --git a/Engine/Source/Runtime/TraceLog/Public/Trace/Trace.h b/Engine/Source/Runtime/TraceLog/Public/Trace/Trace.h
--- a/Engine/Source/Runtime/TraceLog/Public/Trace/Trace.h
+++ b/Engine/Source/Runtime/TraceLog/Public/Trace/Trace.h
@@ -154,6 +154,7 @@ struct FInitializeDesc
 {
 	uint32				TailSizeBytes		= 4 << 20; // can be set to 0 to disable the tail buffer
 	uint32				ThreadSleepTimeInMS = 0;
+	uint32				BlockPoolMaxSize	= UE_TRACE_BLOCK_POOL_MAXSIZE;
 	bool				bUseWorkerThread	= true;
 	bool				bUseImportantCache	= true;
 	uint32				SessionGuid[4]		= {0,0,0,0}; // leave as zero to generate random
@@ -176,13 +177,50 @@ struct FChannelInfo
 
 class FChannel;
 
-typedef void*		AllocFunc(SIZE_T, uint32);
-typedef void		FreeFunc(void*, SIZE_T);
-typedef void		ChannelIterFunc(const ANSICHAR*, bool, void*);
-/*  The callback provides information about a channel and a user provided pointer.
-	Returning false from the callback will stop the enumeration */
+/**
+ * Allocate memory callback
+ * @param Size Size to allocate
+ * @param Alignment Alignment of memory
+ * @return Pointer to allocated memory
+ */
+typedef void*		AllocFunc(SIZE_T Size, uint32 Alignment);
+
+/**
+ * Free memory callback
+ * @param Ptr Memory to free
+ * @param Size Size of memory to free
+ */
+typedef void		FreeFunc(void* Ptr, SIZE_T Size);
+
+/**  
+ * The callback provides information about a channel and a user provided pointer.
+ * @param Name Name of channel
+ * @param State Enabled state of channel
+ * @param User User data passed to the callback
+ */
+typedef void		ChannelIterFunc(const ANSICHAR* Name, bool State, void* User);
+
+/**  
+ * The callback provides information about a channel and a user provided pointer.
+ * @param Info Information about the channel
+ * @param User User data passed to the callback
+ * @return Returning false from the callback will stop the enumeration 
+ */
 typedef bool		ChannelIterCallback(const FChannelInfo& Info, void*/*User*/);
+
+/**
+ * User defined write callback.
+ * @param Handle User defined handle passed to the function
+ * @param Data Pointer to data to write
+ * @param Size Size of data to write
+ * @return True if all data could be written correctly, false if error occurred.
+ */
 typedef bool		IoWriteFunc(UPTRINT Handle, const void* Data, uint32 Size);
+
+/**
+ * User defined close callback.
+ * @param Handle User defined handle passed to the function
+ */
 typedef void		IoCloseFunc(UPTRINT Handle);	
 
 struct FStatistics
@@ -206,32 +244,192 @@ struct FSendFlags
 	static const uint16 _Reserved	= 1 << 15;	// this bit is used internally
 };
 
+/**
+ * Set optional allocation and free methods to use. If not set TraceLog will fall back to
+ * default platform allocation methods.
+ * @param Alloc Callback for allocations
+ * @param Free Callback for free
+ */
 UE_TRACE_API void	SetMemoryHooks(AllocFunc Alloc, FreeFunc Free) UE_TRACE_IMPL();
+
+/**
+ * Set optional callback to use for critical messages. See \ref OnMessageFunc for details.
+ * @param MessageFunc Function to call for critical messages
+ */
 UE_TRACE_API void	SetMessageCallback(OnMessageFunc* MessageFunc) UE_TRACE_IMPL();
+
+/**
+ * Set optional on update callbacks. If set issued after every update and once after initalization.
+ * @param UpdateFunc Function to call after every update
+ */
 UE_TRACE_API void	SetUpdateCallback(OnUpdateFunc* UpdateFunc) UE_TRACE_IMPL();
+
+/**
+ * Initalize TraceLog library.
+ * @param Desc Initalization options
+ */
 UE_TRACE_API void	Initialize(const FInitializeDesc& Desc) UE_TRACE_IMPL();
+
+/**
+ * Manually start worker thread if library is initalized without worker thread.
+ */
 UE_TRACE_API void	StartWorkerThread() UE_TRACE_IMPL();
+
+/**
+ * Call when application is exiting. Notifies TraceLog that the worker thread can 
+ * dissapear at any time and clears block pool limits. Tracing is still possible.
+ */
+UE_TRACE_API void	Exit() UE_TRACE_IMPL();
+
+/**
+ * Shuts down the library completely and frees resources. After this tracing will 
+ * not be possible.
+ */
 UE_TRACE_API void	Shutdown() UE_TRACE_IMPL();
+
+/**
+ * Notifies TraceLog about a critical failure. Disables all tracing by muting all channels.
+ */
 UE_TRACE_API void	Panic() UE_TRACE_IMPL();
+
+/**
+ * Manually update TraceLog if no worker thread is running. Only one thread (including the 
+ * worker thread) is able to enter the update method at any point.
+ */
 UE_TRACE_API void	Update() UE_TRACE_IMPL();
+
+/**
+ * Fetches tracked telemetry from the library.
+ * @param Out Pointer to struct where telemetry values will be written to.
+ */ 
 UE_TRACE_API void	GetStatistics(FStatistics& Out) UE_TRACE_IMPL();
+
+/**
+ * Setup TraceLog to output to remote host using a socket connection, to take effect next update. 
+ * Will fail if another pending output has been queued or the host is unreachable.
+ * @param Host Target hostname or ip
+ * @param Port Target port
+ * @param Flags Options for the connection
+ * @return True on successful connection to host, false otherwise
+ */ 
 UE_TRACE_API bool	SendTo(const TCHAR* Host, uint32 Port=0, uint16 Flags=FSendFlags::None) UE_TRACE_IMPL(false);
+
+/**
+ * Setup TraceLog to output to a new or existing file, to take effect next update. Will fail if another pending 
+ * output has been queued or if the file location is not writeable.
+ * @param Path Target path
+ * @param Flags Options for the connection
+ * @return True if the file could be opened or created correctly, false otherwise
+ */
 UE_TRACE_API bool	WriteTo(const TCHAR* Path, uint16 Flags=FSendFlags::None) UE_TRACE_IMPL(false);
+
+/**
+ * Setup TraceLog to output to user defined callback, to take effect next update. Will fail if another pending
+ * output has been queued.
+ * @param WriteFunc Function to call when writing data
+ * @param CloseFunc Function to call when output is closed
+ * @param Flags Options for the connection
+ * @return True if successful, false otherwise
+ */
 UE_TRACE_API bool	RelayTo(UPTRINT InHandle, IoWriteFunc WriteFunc, IoCloseFunc CloseFunc, uint16 Flags = FSendFlags::None) UE_TRACE_IMPL(false);
+
+/**
+ * Immediately write contents of tail buffers and important events to a new or existing file.
+ * @param Path Target path
+ */
 UE_TRACE_API bool	WriteSnapshotTo(const TCHAR* Path) UE_TRACE_IMPL(false);
+
+/**
+ * Immediately write contents of tail buffers and important events to a remote host.
+ * @param Host Target hostname or ip
+ * @param Port Target port
+ */
 UE_TRACE_API bool	SendSnapshotTo(const TCHAR* Host, uint32 Port) UE_TRACE_IMPL(false);
+
+/**
+ * Checks if TraceLog currently has an output. Note that trace events can still be recorded
+ * and saved in tail buffers regardless if an output is active.
+ * @return True if an output is active, false if not
+ */
 UE_TRACE_API bool	IsTracing() UE_TRACE_IMPL(false);
+
+/**
+ * Checks if TraceLog currently has an output and return session and trace GUIDs of active trace stream. 
+ * Note that trace events can still be recorded and saved in tail buffers regardless if an output is active.
+ * @param OutSessionGuid If output is active current session GUID will be written, otherwise it will be unchanged
+ * @param OutTraceGuid If output is active current trace GUID will be written, otherwise it will be unchanged
+ * @return True if an output is active, false if not
+ */
 UE_TRACE_API bool	IsTracingTo(uint32 (&OutSessionGuid)[4], uint32 (&OutTraceGuid)[4]) UE_TRACE_IMPL(false);
+
+/**
+ * Stops current output if any is active.This will fail there is no active output or if there is already a queued
+ * output to be started.
+ * @return True if any output was stopped, false if no output was active or there is a queued output to be started.
+ */
 UE_TRACE_API bool	Stop() UE_TRACE_IMPL(false);
+
+/**
+ * Checks if a string is valid channel name.
+ * @param ChanelName String to check
+ * @return True if there is a registered channel object with the same name.
+ */
 UE_TRACE_API bool	IsChannel(const TCHAR* ChannelName) UE_TRACE_IMPL(false);
+
+/**
+ * Toggles channel to control output of events.
+ * @param ChanelName Name of channel to control
+ * @param bEnabled Set if channel should be enabled or disabled
+ * @return Final channel state. If the channel name was not found return false.
+ */
 UE_TRACE_API bool	ToggleChannel(const TCHAR* ChannelName, bool bEnabled) UE_TRACE_IMPL(false);
+
+/**
+ * Enumerates registered channels.
+ * @param IterFunc Function to call for each channel
+ * @param User Optional pointer to user data to pass to callback
+ */
 UE_TRACE_API void	EnumerateChannels(ChannelIterFunc IterFunc, void* User) UE_TRACE_IMPL();
+
+/**
+ * Enumerates registered channels.
+ * @param IterFunc Function to call for each channel
+ * @param User Optional pointer to user data to pass to callback
+ */
 UE_TRACE_API void	EnumerateChannels(ChannelIterCallback IterFunc, void* User) UE_TRACE_IMPL();
+
+/**
+ * Register a new thread in Trace. This is a requirement before tracing anything from the thrad.
+ * @param Name Display name of thread
+ * @param SystemId Platform specific thread id
+ * @param SortHint Suggestion on how to order thread in presentation
+ */
 UE_TRACE_API void	ThreadRegister(const TCHAR* Name, uint32 SystemId, int32 SortHint) UE_TRACE_IMPL();
+
+/**
+ * Define a group of threads with similar use. Any thread created within this call and /ref ThreadGroupEnd
+ * will be grouped together.
+ * @param Name Display name of group
+ */
 UE_TRACE_API void	ThreadGroupBegin(const TCHAR* Name) UE_TRACE_IMPL();
+
+/**
+ * End a group of threads with similar use. See /ref ThreadGroupBegin. 
+ */
 UE_TRACE_API void	ThreadGroupEnd() UE_TRACE_IMPL();
 
+/**
+ * Attempts to find the corresponding channel object given a channel name.
+ * @param ChannelName Name to search for
+ * @return Pointer to channel object, or null if no channel matching the name have been registered
+ */
 UE_TRACE_API FChannel* FindChannel(const TCHAR* ChannelName) UE_TRACE_IMPL(nullptr);
+
+/**
+ * Attempts to find the corresponding channel object given a channel id.
+ * @param ChannelId Id to search for
+ * @return Pointer to channel object, or null if no channel with matching id have been registered
+ */
 UE_TRACE_API FChannel* FindChannel(FChannelId ChannelId) UE_TRACE_IMPL(nullptr);
 
 } // namespace Trace
