[TraceLog]
* Changed TEventRef::GetHash to use if constexpr instead of template specialization

diff --git a/Engine/Source/Runtime/TraceLog/Public/Trace/Trace.h b/Engine/Source/Runtime/TraceLog/Public/Trace/Trace.h
--- a/Engine/Source/Runtime/TraceLog/Public/Trace/Trace.h
+++ b/Engine/Source/Runtime/TraceLog/Public/Trace/Trace.h
@@ -67,36 +67,22 @@ struct TEventRef
 	IdType Id;
 	uint32 RefTypeId;
 
-	uint64 GetHash() const;
+	uint64 GetHash() const
+	{
+		if constexpr (std::is_same_v<IdType, uint64>)
+		{
+			return (uint64(RefTypeId) << 32) ^ Id;
+		}
+		else
+		{
+			return (uint64(RefTypeId) << 32) | Id;
+		}
+	}
 
 private:
 	TEventRef() = delete;
 };
 
-template <>
-inline uint64 TEventRef<uint8>::GetHash() const
-{
-	return (uint64(RefTypeId) << 32) | Id;
-}
-
-template <>
-inline uint64 TEventRef<uint16>::GetHash() const
-{
-	return (uint64(RefTypeId) << 32) | Id;
-}
-
-template <>
-inline uint64 TEventRef<uint32>::GetHash() const
-{
-	return (uint64(RefTypeId) << 32) | Id;
-}
-
-template <>
-inline uint64 TEventRef<uint64>::GetHash() const
-{
-	return (uint64(RefTypeId) << 32) ^ Id;
-}
-
 typedef TEventRef<uint8> FEventRef8;
 typedef TEventRef<uint16> FEventRef16;
 typedef TEventRef<uint32> FEventRef32;
