Improvements to handling compression in TraceLog

* If compression fails, fall back to sending data uncompressed.
* Emit error messages in cases of compression failures.

#rnx

diff --git a/Engine/Source/Runtime/TraceLog/Private/Trace/Codec.cpp b/Engine/Source/Runtime/TraceLog/Private/Trace/Codec.cpp
--- a/Engine/Source/Runtime/TraceLog/Private/Trace/Codec.cpp
+++ b/Engine/Source/Runtime/TraceLog/Private/Trace/Codec.cpp
@@ -2,6 +2,7 @@
 
 #include "CoreTypes.h"
 #include "Platform.h"
+#include "Message.h"
 
 THIRD_PARTY_INCLUDES_START
 #if defined(_MSC_VER)
@@ -33,25 +34,39 @@ int32 Encode(const void* Src, int32 SrcSize, void* Dest, int32 DestSize)
 #if TRACE_PRIVATE_MINIMAL_ENABLED
 	FProfilerScope _(__func__);
 #endif
-	return TRACE_PRIVATE_LZ4_NAMESPACE LZ4_compress_fast(
+	int Result = TRACE_PRIVATE_LZ4_NAMESPACE LZ4_compress_fast(
 		(const char*)Src,
 		(char*)Dest,
 		SrcSize,
 		DestSize,
 		1 // increase by 1 for small speed increase
 	);
+	if (!Result)
+	{
+#if UE_TRACE_ENABLED 
+		UE_TRACE_MESSAGE_F(CompressionError, "LZ4 failed to compress %d bytes (dest %d bytes).", SrcSize, DestSize);
+#endif
+	}
+	return Result;
 }
 
 ////////////////////////////////////////////////////////////////////////////////
 int32 EncodeNoInstr(const void* Src, int32 SrcSize, void* Dest, int32 DestSize)
 {
-	return TRACE_PRIVATE_LZ4_NAMESPACE LZ4_compress_fast(
+	int Result = TRACE_PRIVATE_LZ4_NAMESPACE LZ4_compress_fast(
 		(const char*)Src,
 		(char*)Dest,
 		SrcSize,
 		DestSize,
 		1 // increase by 1 for small speed increase
 	);
+	if (!Result)
+	{
+		// Cannot emit error message here since it might cause
+		// events.
+		//PLATFORM_BREAK();
+	}
+	return Result;
 }
 
 ////////////////////////////////////////////////////////////////////////////////
diff --git a/Engine/Source/Runtime/TraceLog/Private/Trace/Important/Cache.cpp b/Engine/Source/Runtime/TraceLog/Private/Trace/Important/Cache.cpp
--- a/Engine/Source/Runtime/TraceLog/Private/Trace/Important/Cache.cpp
+++ b/Engine/Source/Runtime/TraceLog/Private/Trace/Important/Cache.cpp
@@ -1,10 +1,11 @@
 // Copyright Epic Games, Inc. All Rights Reserved.
 
+#include "GenericPlatform/GenericPlatformMisc.h"
 #include "Trace/Config.h"
-#include "Trace/Platform.h"
 
 #if TRACE_PRIVATE_MINIMAL_ENABLED && TRACE_PRIVATE_ALLOW_IMPORTANTS
 
+#include "Trace/Platform.h"
 #include "Trace/Detail/Protocol.h"
 #include "Trace/Detail/Transport.h"
 #include "Trace/Trace.h"
@@ -87,19 +88,38 @@ static void Writer_CacheCommit(const FCacheBuffer* Collector)
 		GCacheActiveBuffer = NewBuffer;
 	}
 
-	uint32 Used = GCacheActiveBuffer->Size - GCacheActiveBuffer->Remaining;
-	auto* Packet = (FTidPacketEncoded*)(GCacheActiveBuffer->Data + Used);
-	uint32 OutputSize = Encode(Collector->Data, InputSize, Packet->Data, EncodeMaxSize);
+	const uint32 ActiveBufferOffset = GCacheActiveBuffer->Size - GCacheActiveBuffer->Remaining;
+	uint32 PacketTotalSize = 0;
+	auto* Packet = (FTidPacketEncoded*)(GCacheActiveBuffer->Data + ActiveBufferOffset);
+	const int32 CompressedSize = Encode(Collector->Data, InputSize, Packet->Data, EncodeMaxSize);
+	if (CompressedSize > 0)
+	{
+		Packet->PacketSize = uint16(CompressedSize + sizeof(FTidPacketEncoded));
+		Packet->ThreadId = FTidPacketBase::EncodedMarker | uint16(ETransportTid::Importants);
+		Packet->DecodedSize = uint16(InputSize);
+
+		PacketTotalSize = sizeof(FTidPacketEncoded) + CompressedSize;
+	}
+	else
+	{
+		// Compression failed. We can reuse the same buffer for uncompressed packet
+		// since the worst case size for compression is (size + compressionoverhead)
+		// and FTidPacket < FTidPackedEncoded.
+		auto* UncompressedPacket = (FTidPacket*)(GCacheActiveBuffer->Data + ActiveBufferOffset);
+		UncompressedPacket->ThreadId = uint16(ETransportTid::Importants);
+#if UE_TRACE_PACKET_VERIFICATION
+		Packet->ThreadId |= FTidPacketBase::Verification;
+#endif
+		UncompressedPacket->PacketSize = uint16(InputSize) + sizeof(FTidPacket);
+		::memcpy(UncompressedPacket->Data, Collector->Data, InputSize);
 
-	Packet->PacketSize = uint16(OutputSize + sizeof(FTidPacketEncoded));
-	Packet->ThreadId = FTidPacketBase::EncodedMarker | uint16(ETransportTid::Importants);
-	Packet->DecodedSize = uint16(InputSize);
+		PacketTotalSize = sizeof(FTidPacket) + InputSize;
+	}
 
-	Used = sizeof(FTidPacketEncoded) + OutputSize;
-	GCacheActiveBuffer->Remaining -= Used;
+	GCacheActiveBuffer->Remaining -= PacketTotalSize;
 
 #if TRACE_PRIVATE_STATISTICS
-	GTraceStatistics.CacheUsed += Used;
+	GTraceStatistics.CacheUsed += PacketTotalSize;
 #endif
 }
 
diff --git a/Engine/Source/Runtime/TraceLog/Private/Trace/Message.h b/Engine/Source/Runtime/TraceLog/Private/Trace/Message.h
--- a/Engine/Source/Runtime/TraceLog/Private/Trace/Message.h
+++ b/Engine/Source/Runtime/TraceLog/Private/Trace/Message.h
@@ -4,6 +4,7 @@
 
 #include "Trace/Config.h"
 #include "Trace/Trace.h"
+#include <stdio.h>
 
 #if TRACE_PRIVATE_MINIMAL_ENABLED
 
diff --git a/Engine/Source/Runtime/TraceLog/Private/Trace/Tail.cpp b/Engine/Source/Runtime/TraceLog/Private/Trace/Tail.cpp
--- a/Engine/Source/Runtime/TraceLog/Private/Trace/Tail.cpp
+++ b/Engine/Source/Runtime/TraceLog/Private/Trace/Tail.cpp
@@ -1,10 +1,10 @@
 // Copyright Epic Games, Inc. All Rights Reserved.
 
-#include "Platform.h"
 #include "Trace/Config.h"
 
 #if TRACE_PRIVATE_MINIMAL_ENABLED
 
+#include "Platform.h"
 #include "Trace/Detail/Atomic.h"
 #include "Trace/Detail/Protocol.h"
 #include "Trace/Detail/Transport.h"
@@ -56,6 +56,8 @@ public:
 	void			IterateRanges(CallbackType&& Callback);
 	template <typename PacketType>
 	PacketType*		Append(uint32 InSize);
+	template <typename PacketType>
+	void			BackUp(uint32 InSize);
 	void			BackUp(uint32 InSize);
 
 private:
@@ -149,6 +151,13 @@ PacketType* FPacketRing::Append(uint32 InSize)
 	return static_cast<PacketType*>(Ptr);
 }
 
+////////////////////////////////////////////////////////////////////////////////
+template <typename PacketType>
+void FPacketRing::BackUp(uint32 InSize)
+{
+	Cursor -= (InSize + sizeof(PacketType));
+}
+
 ////////////////////////////////////////////////////////////////////////////////
 void FPacketRing::BackUp(uint32 InSize)
 {
@@ -235,36 +244,45 @@ void Writer_TailAppend(uint32 ThreadId, uint8* __restrict Data, uint32 Size)
 
 	// Smaller buffers usually aren't redundant enough to benefit from being
 	// compressed. They often end up being larger.
-	if (Size <= 384)
+	if (Size > 384)
 	{
-		auto* Packet = GPacketRing.Append<FTidPacket>(Size);
+		uint32 EncodeMaxSize = GetEncodeMaxSize(Size);
+		auto* Packet = GPacketRing.Append<FTidPacketEncoded>(EncodeMaxSize);
 		Packet->ThreadId = uint16(ThreadId);
+		Packet->ThreadId |= FTidPacketBase::EncodedMarker;
 #if UE_TRACE_PACKET_VERIFICATION
 		Packet->ThreadId |= FTidPacketBase::Verification;
 #endif
-		::memcpy(Packet->Data, Data, Size);
+		Packet->DecodedSize = uint16(Size);
+
+		int32 EncodedSize = Encode(Data, Size, Packet->Data, EncodeMaxSize);
+		if (EncodedSize > 0)
+		{
+			// Reclaim buffer space (worst case size - actual size)
+			uint32 BackUp = EncodeMaxSize - EncodedSize;
+			GPacketRing.BackUp(BackUp);
+			Packet->PacketSize -= uint16(BackUp);
 
-		Writer_SendDataRaw(Packet, Packet->PacketSize);
+			Writer_SendDataRaw(Packet, Packet->PacketSize);
 #if UE_TRACE_PACKET_VERIFICATION
-		const uint64 Serial = GPacketSerial++;
-		Writer_SendDataRaw(&Serial, sizeof(uint64));
+			const uint64 Serial = GPacketSerial++;
+			Writer_SendDataRaw(&Serial, sizeof(uint64));
 #endif
-		return;
+			return;
+		}
+		else
+		{
+			// Compression failed, back out of entire packet
+			GPacketRing.BackUp<FTidPacketEncoded>(EncodeMaxSize);
+		}
 	}
 
-	uint32 EncodeMaxSize = GetEncodeMaxSize(Size);
-	auto* Packet = GPacketRing.Append<FTidPacketEncoded>(EncodeMaxSize);
+	auto* Packet = GPacketRing.Append<FTidPacket>(Size);
 	Packet->ThreadId = uint16(ThreadId);
-	Packet->ThreadId |= FTidPacketBase::EncodedMarker;
 #if UE_TRACE_PACKET_VERIFICATION
 	Packet->ThreadId |= FTidPacketBase::Verification;
 #endif
-	Packet->DecodedSize = uint16(Size);
-
-	uint32 EncodeSize = Encode(Data, Size, Packet->Data, EncodeMaxSize);
-	uint32 BackUp = EncodeMaxSize - EncodeSize;
-	GPacketRing.BackUp(BackUp);
-	Packet->PacketSize -= uint16(BackUp);
+	::memcpy(Packet->Data, Data, Size);
 
 	Writer_SendDataRaw(Packet, Packet->PacketSize);
 #if UE_TRACE_PACKET_VERIFICATION
diff --git a/Engine/Source/Runtime/TraceLog/Private/Trace/Writer.cpp b/Engine/Source/Runtime/TraceLog/Private/Trace/Writer.cpp
--- a/Engine/Source/Runtime/TraceLog/Private/Trace/Writer.cpp
+++ b/Engine/Source/Runtime/TraceLog/Private/Trace/Writer.cpp
@@ -1,10 +1,10 @@
 // Copyright Epic Games, Inc. All Rights Reserved.
 
-#include "Message.h"
 #include "Trace/Config.h"
 
 #if TRACE_PRIVATE_MINIMAL_ENABLED
 
+#include "Message.h"
 #include "Platform.h"
 #include "Trace/Detail/Atomic.h"
 #include "Trace/Detail/Channel.h"
@@ -423,7 +423,49 @@ void Writer_SendDataNoInstr(uint32 ThreadId, uint8* __restrict Data, uint32 Size
 
 	// Smaller buffers usually aren't redundant enough to benefit from being
 	// compressed. They often end up being larger.
-	if (Size <= 384)
+	if (Size > 384)
+	{
+		constexpr uint32 MaxEncodedBuffer = 8192;
+
+		// Buffer size is expressed as "A + B" where A is a maximum expected
+		// input size (i.e. at least GPoolBlockSize) and B is LZ4 overhead as
+		// per LZ4_COMPRESSBOUND.
+		TTidPacketEncoded<MaxEncodedBuffer + 64> Packet;
+
+		if (Size > MaxEncodedBuffer)
+		{
+			PLATFORM_BREAK();
+		}
+
+		Packet.DecodedSize = uint16(Size);
+		Packet.PacketSize = uint16(EncodeNoInstr(Data, Packet.DecodedSize, Packet.Data, sizeof(Packet.Data)));
+
+		if (Packet.PacketSize > 0 &&
+			Packet.PacketSize < Packet.DecodedSize)
+		{
+			Packet.ThreadId = FTidPacketBase::EncodedMarker;
+			Packet.ThreadId |= uint16(ThreadId & FTidPacketBase::ThreadIdMask);
+#if UE_TRACE_PACKET_VERIFICATION
+			Packet.ThreadId |= FTidPacketBase::Verification;
+#endif
+			Packet.PacketSize += sizeof(FTidPacketEncoded);
+
+			Writer_SendDataImplNoInstr(&Packet, Packet.PacketSize);
+
+#if UE_TRACE_PACKET_VERIFICATION
+			const uint64 Serial = GPacketSerial++;
+			Writer_SendDataImplNoInstr(&Serial, sizeof(uint64));
+#endif
+			return;
+		}
+		else
+		{
+			PLATFORM_BREAK();
+			// Intentional fall through to write uncompressed data...
+		}
+	}
+
+	// Write uncompressed data.
 	{
 		Data -= sizeof(FTidPacket);
 		Size += sizeof(FTidPacket);
@@ -440,29 +482,7 @@ void Writer_SendDataNoInstr(uint32 ThreadId, uint8* __restrict Data, uint32 Size
 		const uint64 Serial = GPacketSerial++;
 		Writer_SendDataImplNoInstr(&Serial, sizeof(uint64));
 #endif
-		return;
 	}
-
-	// Buffer size is expressed as "A + B" where A is a maximum expected
-	// input size (i.e. at least GPoolBlockSize) and B is LZ4 overhead as
-	// per LZ4_COMPRESSBOUND.
-	TTidPacketEncoded<8192 + 64> Packet;
-
-	Packet.ThreadId = FTidPacketBase::EncodedMarker;
-	Packet.ThreadId |= uint16(ThreadId & FTidPacketBase::ThreadIdMask);
-#if UE_TRACE_PACKET_VERIFICATION
-	Packet.ThreadId |= FTidPacketBase::Verification;
-#endif
-	Packet.DecodedSize = uint16(Size);
-	Packet.PacketSize = uint16(EncodeNoInstr(Data, Packet.DecodedSize, Packet.Data, sizeof(Packet.Data)));
-	Packet.PacketSize += sizeof(FTidPacketEncoded);
-
-	Writer_SendDataImplNoInstr(&Packet, Packet.PacketSize);
-
-#if UE_TRACE_PACKET_VERIFICATION
-	const uint64 Serial = GPacketSerial++;
-	Writer_SendDataImplNoInstr(&Serial, sizeof(uint64));
-#endif
 }
 
 ////////////////////////////////////////////////////////////////////////////////
@@ -477,8 +497,51 @@ void Writer_SendData(uint32 ThreadId, uint8* __restrict Data, uint32 Size)
 
 	// Smaller buffers usually aren't redundant enough to benefit from being
 	// compressed. They often end up being larger.
-	if (Size <= 384)
+	if (Size > 384)
+	{
+		constexpr uint32 MaxEncodedBuffer = 8192;
+
+		// Buffer size is expressed as "A + B" where A is a maximum expected
+		// input size (i.e. at least GPoolBlockSize) and B is LZ4 overhead as
+		// per LZ4_COMPRESSBOUND.
+		TTidPacketEncoded<MaxEncodedBuffer + 64> Packet;
+
+		if (Size > MaxEncodedBuffer)
+		{
+			PLATFORM_BREAK();
+		}
+
+		Packet.DecodedSize = uint16(Size);
+		Packet.PacketSize = uint16(Encode(Data, Packet.DecodedSize, Packet.Data, sizeof(Packet.Data)));
+
+		if (Packet.PacketSize > 0 &&
+			Packet.PacketSize < Packet.DecodedSize)
+		{
+			Packet.ThreadId = uint16(ThreadId & FTidPacketBase::ThreadIdMask);
+			Packet.ThreadId |= FTidPacketBase::EncodedMarker;
+#if UE_TRACE_PACKET_VERIFICATION
+			Packet.ThreadId |= FTidPacketBase::Verification;
+#endif
+			Packet.PacketSize += sizeof(FTidPacketEncoded);
+
+			Writer_SendDataImpl(&Packet, Packet.PacketSize);
+#if UE_TRACE_PACKET_VERIFICATION
+			const uint64 Serial = GPacketSerial++;
+			Writer_SendDataImpl(&Serial, sizeof(uint64));
+#endif
+			return;
+		}
+		else
+		{
+			PLATFORM_BREAK();
+			// Intentional fall through to write uncompressed data...
+		}
+	}
+
+	// Write uncompressed data.
 	{
+		// When setting up the buffer we have left space for
+		// a uncompressed packet header ahead of the data.
 		Data -= sizeof(FTidPacket);
 		Size += sizeof(FTidPacket);
 		auto* Packet = (FTidPacket*)Data;
@@ -494,29 +557,7 @@ void Writer_SendData(uint32 ThreadId, uint8* __restrict Data, uint32 Size)
 		const uint64 Serial = GPacketSerial++;
 		Writer_SendDataImpl(&Serial, sizeof(uint64));
 #endif
-		return;
 	}
-
-	// Buffer size is expressed as "A + B" where A is a maximum expected
-	// input size (i.e. at least GPoolBlockSize) and B is LZ4 overhead as
-	// per LZ4_COMPRESSBOUND.
-	TTidPacketEncoded<8192 + 64> Packet;
-
-	Packet.ThreadId = FTidPacketBase::EncodedMarker;
-	Packet.ThreadId |= uint16(ThreadId & FTidPacketBase::ThreadIdMask);
-#if UE_TRACE_PACKET_VERIFICATION
-	Packet.ThreadId |= FTidPacketBase::Verification;
-#endif
-	Packet.DecodedSize = uint16(Size);
-	Packet.PacketSize = uint16(Encode(Data, Packet.DecodedSize, Packet.Data, sizeof(Packet.Data)));
-	Packet.PacketSize += sizeof(FTidPacketEncoded);
-
-	Writer_SendDataImpl(&Packet, Packet.PacketSize);
-
-#if UE_TRACE_PACKET_VERIFICATION
-	const uint64 Serial = GPacketSerial++;
-	Writer_SendDataImpl(&Serial, sizeof(uint64));
-#endif
 }
 
 ////////////////////////////////////////////////////////////////////////////////
diff --git a/Engine/Source/Runtime/TraceLog/Public/Trace/Trace.h b/Engine/Source/Runtime/TraceLog/Public/Trace/Trace.h
--- a/Engine/Source/Runtime/TraceLog/Public/Trace/Trace.h
+++ b/Engine/Source/Runtime/TraceLog/Public/Trace/Trace.h
@@ -115,6 +115,7 @@ enum class EMessageType : uint8
 	EstablishError,
 	FileOpenError,
 	WriterError,
+	CompressionError,
 	// Fatal errors should cause application to stop
 	FatalStart			= 0x40,
 	OOMFatal,
