A less enthusiastic backoff

#rnx

diff --git a/Engine/Source/Runtime/TraceLog/Private/Trace/BlockPool.cpp b/Engine/Source/Runtime/TraceLog/Private/Trace/BlockPool.cpp
--- a/Engine/Source/Runtime/TraceLog/Private/Trace/BlockPool.cpp
+++ b/Engine/Source/Runtime/TraceLog/Private/Trace/BlockPool.cpp
@@ -111,13 +111,27 @@ static FPoolBlockList Writer_AllocateBlockList(uint32 PageSize)
 ////////////////////////////////////////////////////////////////////////////////
 FWriteBuffer* Writer_AllocateBlockFromPool()
 {
+	int32 ThrottleRestore = -1;
+	ON_SCOPE_EXIT {
+		if (ThrottleRestore >= 0)
+		{
+			ThreadUnthrottle(ThrottleRestore);
+		}
+	};
+
+	auto ApplyThrottle = [&ThrottleRestore, Count=uint32(0)] () mutable {
+		if (ThrottleRestore < 0 && Count++ > 0)
+		{
+			ThrottleRestore = ThreadThrottle();
+		}
+	};
+
 	// Fetch a new buffer
 	uint64 StallStart = 0;
 	FWriteBuffer* Ret;
 	while (true)
 	{
-		int32 ThrottleRestore = ThreadThrottle();
-		ON_SCOPE_EXIT { ThreadUnthrottle(ThrottleRestore); };
+		ApplyThrottle();
 
 		// First we'll try one from the free list
 		FWriteBuffer* Owned = AtomicLoadRelaxed(&GPoolFreeList);
