[TraceLog]
- Fix races on stats

diff --git a/Engine/Source/Runtime/TraceLog/Private/Trace/Important/Cache.cpp b/Engine/Source/Runtime/TraceLog/Private/Trace/Important/Cache.cpp
--- a/Engine/Source/Runtime/TraceLog/Private/Trace/Important/Cache.cpp
+++ b/Engine/Source/Runtime/TraceLog/Private/Trace/Important/Cache.cpp
@@ -6,6 +6,7 @@
 #if TRACE_PRIVATE_MINIMAL_ENABLED && TRACE_PRIVATE_ALLOW_IMPORTANTS
 
 #include "Trace/Platform.h"
+#include "Trace/Detail/Atomic.h"
 #include "Trace/Detail/Protocol.h"
 #include "Trace/Detail/Transport.h"
 #include "Trace/Trace.h"
@@ -55,7 +56,7 @@ static FCacheBuffer* Writer_CacheCreateBuffer(uint32 Size)
 	const uint32 CacheBufferSize = sizeof(FCacheBuffer) + Size;
 	void* Block = Writer_MemoryAllocate(CacheBufferSize, alignof(FCacheBuffer));
 #if TRACE_PRIVATE_STATISTICS
-	GTraceStatistics.CacheAllocated += CacheBufferSize;
+	AtomicAddRelaxed(&GTraceStatistics.CacheAllocated, CacheBufferSize);
 #endif
 	auto* Buffer = (FCacheBuffer*)Block;
 	Buffer->Size = Size;
@@ -75,7 +76,7 @@ static void Writer_CacheCommit(const FCacheBuffer* Collector)
 	if (EncodeMaxSize + sizeof(FTidPacketEncoded) > GCacheActiveBuffer->Remaining)
 	{
 #if TRACE_PRIVATE_STATISTICS
-		GTraceStatistics.CacheWaste += GCacheActiveBuffer->Remaining;
+		AtomicAddRelaxed(&GTraceStatistics.CacheWaste, GCacheActiveBuffer->Remaining);
 #endif
 
 		// Retire active buffer
@@ -119,7 +120,7 @@ static void Writer_CacheCommit(const FCacheBuffer* Collector)
 	GCacheActiveBuffer->Remaining -= PacketTotalSize;
 
 #if TRACE_PRIVATE_STATISTICS
-	GTraceStatistics.CacheUsed += PacketTotalSize;
+	AtomicAddRelaxed(&GTraceStatistics.CacheUsed, PacketTotalSize);
 #endif
 }
 
@@ -200,13 +201,16 @@ void Writer_InitializeCache()
 ////////////////////////////////////////////////////////////////////////////////
 void Writer_ShutdownCache()
 {
+#if TRACE_PRIVATE_STATISTICS
+	uint32 TotalCacheToSubtract = 0;
+#endif
 	for (FCacheBuffer* Buffer = GCacheHeadBuffer; Buffer != nullptr;)
 	{
 		FCacheBuffer* Next = Buffer->Next;
 		const uint32 CacheBufferSize = Buffer->Size + sizeof(FCacheBuffer);
 		Writer_MemoryFree(Buffer, CacheBufferSize);
 #if TRACE_PRIVATE_STATISTICS
-		GTraceStatistics.CacheAllocated -= CacheBufferSize;
+		TotalCacheToSubtract += CacheBufferSize;
 #endif
 		Buffer = Next;
 	}
@@ -214,8 +218,7 @@ void Writer_ShutdownCache()
 	Writer_MemoryFree(GCacheActiveBuffer, GCacheBufferSize);
 	Writer_MemoryFree(GCacheCollector, GCacheCollectorSize);
 #if TRACE_PRIVATE_STATISTICS
-	GTraceStatistics.CacheAllocated -= GCacheBufferSize;
-	GTraceStatistics.CacheAllocated -= GCacheCollectorSize;
+	AtomicSubRelaxed(&GTraceStatistics.CacheAllocated, TotalCacheToSubtract + GCacheBufferSize + GCacheCollectorSize);
 #endif
 }
 
diff --git a/Engine/Source/Runtime/TraceLog/Private/Trace/Important/SharedBuffer.cpp b/Engine/Source/Runtime/TraceLog/Private/Trace/Important/SharedBuffer.cpp
--- a/Engine/Source/Runtime/TraceLog/Private/Trace/Important/SharedBuffer.cpp
+++ b/Engine/Source/Runtime/TraceLog/Private/Trace/Important/SharedBuffer.cpp
@@ -117,7 +117,7 @@ static void Writer_RetireSharedBufferImpl()
 	if (auto SendSize = UPTRINT(GTailBuffer) - UPTRINT(Data) - GTailBuffer->Final)
 	{
 #if TRACE_PRIVATE_STATISTICS
-		GTraceStatistics.BytesTraced += SendSize;
+		AtomicAddRelaxed(&GTraceStatistics.BytesTraced, (uint64)SendSize);
 #endif
 
 		Writer_CacheData(Data, uint32(SendSize));
