[Optim] Add config variables to allow variable trace log write buffer size and prevent writer sleeping. This helps prevent the write thread get saturated under heavy load. Changed variables via build config for editor only

diff --git a/Engine/Source/Runtime/TraceLog/Private/Trace/BlockPool.cpp b/Engine/Source/Runtime/TraceLog/Private/Trace/BlockPool.cpp
--- a/Engine/Source/Runtime/TraceLog/Private/Trace/BlockPool.cpp
+++ b/Engine/Source/Runtime/TraceLog/Private/Trace/BlockPool.cpp
@@ -38,7 +38,7 @@ struct FPoolBlockList
 
 ////////////////////////////////////////////////////////////////////////////////
 #define T_ALIGN alignas(PLATFORM_CACHE_LINE_SIZE)
-static const uint32						GPoolBlockSize		= 4 << 10;
+static const uint32						GPoolBlockSize		= UE_TRACE_BLOCK_SIZE;
 static const uint32						GPoolPageSize		= GPoolBlockSize << 4;
 static const uint32						GPoolInitPageSize	= GPoolBlockSize << 6;
 T_ALIGN static volatile uint32			GPoolTotalAllocated;// = 0
diff --git a/Engine/Source/Runtime/TraceLog/Private/Trace/Writer.cpp b/Engine/Source/Runtime/TraceLog/Private/Trace/Writer.cpp
--- a/Engine/Source/Runtime/TraceLog/Private/Trace/Writer.cpp
+++ b/Engine/Source/Runtime/TraceLog/Private/Trace/Writer.cpp
@@ -421,7 +421,7 @@ void Writer_SendDataNoInstr(uint32 ThreadId, uint8* __restrict Data, uint32 Size
 		return;
 	}
 
-	constexpr uint32 MaxEncodedBuffer = 8192;
+	constexpr uint32 MaxEncodedBuffer = UE_TRACE_BLOCK_SIZE;
 	constexpr uint32 EncodingOverhead = 64;
 
 	// Smaller buffers usually aren't redundant enough to benefit from being
@@ -491,7 +491,7 @@ void Writer_SendData(uint32 ThreadId, uint8* __restrict Data, uint32 Size)
 
 	FProfilerScope _(__func__);
 
-	constexpr uint32 MaxEncodedBuffer = 8192;
+	constexpr uint32 MaxEncodedBuffer = UE_TRACE_BLOCK_SIZE;
 	constexpr uint32 EncodingOverhead = 64;
 
 	// Smaller buffers usually aren't redundant enough to benefit from being
@@ -831,7 +831,7 @@ void Writer_CallbackOnConnect()
 ////////////////////////////////////////////////////////////////////////////////
 static UPTRINT			GWorkerThread;		// = 0;
 static volatile bool	GWorkerThreadQuit;	// = false;
-static uint32			GSleepTimeInMS = 17;
+static uint32			GSleepTimeInMS = UE_TRACE_WRITER_SLEEP_MS;
 static volatile uint32	GUpdateInProgress = 1;	// Don't allow updates until initialized
 
 ////////////////////////////////////////////////////////////////////////////////
diff --git a/Engine/Source/Runtime/TraceLog/Public/Trace/Config.h b/Engine/Source/Runtime/TraceLog/Public/Trace/Config.h
--- a/Engine/Source/Runtime/TraceLog/Public/Trace/Config.h
+++ b/Engine/Source/Runtime/TraceLog/Public/Trace/Config.h
@@ -119,7 +119,21 @@ static_assert(UE_TRACE_MINIMAL_ENABLED == 0 || UE_TRACE_ENABLED_SHIPPING_EXPERIM
  */
 #if !defined(UE_TRACE_BLOCK_POOL_MAXSIZE)
 #	define UE_TRACE_BLOCK_POOL_MAXSIZE 79 << 20
-#endif
+#endif // !defined(UE_TRACE_BLOCK_POOL_MAXSIZE)
+
+/**
+* Default size of each block in the block pool.
+*/
+#if !defined(UE_TRACE_BLOCK_SIZE)
+#	define UE_TRACE_BLOCK_SIZE (4 << 10)
+#endif // UE_TRACE_BLOCK_SIZE
+
+/**
+* Time for writer thread to sleep in between writes
+*/
+#if !defined(UE_TRACE_WRITER_SLEEP_MS)
+#	define UE_TRACE_WRITER_SLEEP_MS (17)
+#endif // UE_TRACE_BLOCK_SIZE
 
 /**
  * Enable packet verification. Only useful when looking for transmission bugs. Note that in order
