[Trace] Improved draining buffers algorithm to first identify the available threads and their buffers and then to only drain these identified buffers. The new buffers (added during execution of Writer_DrainBuffers() will be drained in the next loop.

diff --git a/Engine/Source/Runtime/TraceLog/Private/Trace/TlsBuffer.cpp b/Engine/Source/Runtime/TraceLog/Private/Trace/TlsBuffer.cpp
--- a/Engine/Source/Runtime/TraceLog/Private/Trace/TlsBuffer.cpp
+++ b/Engine/Source/Runtime/TraceLog/Private/Trace/TlsBuffer.cpp
@@ -1,11 +1,11 @@
 // Copyright Epic Games, Inc. All Rights Reserved.
 
-#include "Message.h"
 #include "Trace/Config.h"
 
 #if TRACE_PRIVATE_MINIMAL_ENABLED
 
 #include "HAL/Platform.h" // for PLATFORM_BREAK
+#include "Message.h"
 #include "Platform.h"
 #include "Trace/Detail/Atomic.h"
 #include "Trace/Detail/Writer.inl"
@@ -21,6 +21,8 @@ void				Writer_TailAppend(uint32, uint8* __restrict, uint32);
 FWriteBuffer*		Writer_AllocateBlockFromPool();
 uint32				Writer_GetThreadId();
 void				Writer_FreeBlockListToPool(FWriteBuffer*, FWriteBuffer*);
+void				Writer_MemoryFree(void* Address, uint32 Size);
+void*				Writer_MemoryAllocate(SIZE_T Size, uint32 Alignment);
 extern uint64		GStartCycle;
 extern FStatistics	GTraceStatistics;
 
@@ -33,12 +35,22 @@ UE_TRACE_MINIMAL_EVENT_END()
 
 
 
+////////////////////////////////////////////////////////////////////////////////
+struct FThreadBuffer
+{
+	FWriteBuffer* FirstBuffer;
+	uint32 DrainBufferLimit;
+	uint32 Padding;
+};
+
 ////////////////////////////////////////////////////////////////////////////////
 #define T_ALIGN alignas(PLATFORM_CACHE_LINE_SIZE)
 static FWriteBuffer						GNullWriteBuffer	= { {}, 0, nullptr, nullptr, (uint8*)&GNullWriteBuffer, nullptr, nullptr, 0, 0, 0 };
 thread_local FWriteBuffer*				GTlsWriteBuffer		= &GNullWriteBuffer;
 static FWriteBuffer* __restrict			GActiveThreadList;	// = nullptr;
 T_ALIGN static FWriteBuffer* volatile	GNewThreadList;		// = nullptr;
+static FThreadBuffer*					GThreadBufferArray;	// = nullptr;
+static uint32							GThreadBufferArrayCapacity;	// = 0;
 #undef T_ALIGN
 
 ////////////////////////////////////////////////////////////////////////////////
@@ -117,7 +129,7 @@ static FWriteBuffer* Writer_NextBufferInternal(FWriteBuffer* CurrentBuffer)
 		const uint64 Used = CurrentBuffer->Cursor - ((uint8*)CurrentBuffer - CurrentBuffer->Size);
 		AtomicAddRelaxed(&GTraceStatistics.BytesEmitted, Used);
 	#endif
-		
+
 		NextBuffer->ThreadId = CurrentBuffer->ThreadId;
 		NextBuffer->PrevTimestamp = CurrentBuffer->PrevTimestamp;
 		AtomicStoreRelease(&(CurrentBuffer->NextBuffer), NextBuffer);
@@ -137,6 +149,31 @@ TRACELOG_API FWriteBuffer* Writer_NextBuffer()
 	return Writer_NextBufferInternal(CurrentBuffer);
 }
 
+////////////////////////////////////////////////////////////////////////////////
+static void Writer_ReallocateThreadBufferArray(uint32 NewArrayCapacity)
+{
+	FProfilerScope _(__func__);
+
+	uint32 NewSize = NewArrayCapacity * sizeof(FThreadBuffer);
+	FThreadBuffer* NewArray = static_cast<FThreadBuffer*>(Writer_MemoryAllocate(NewSize, 16));
+#if TRACE_PRIVATE_STATISTICS
+	AtomicAddRelaxed(&GTraceStatistics.FixedBufferAllocated, NewSize);
+#endif
+
+	if (GThreadBufferArrayCapacity > 0)
+	{
+		const uint32 OldSize = GThreadBufferArrayCapacity * sizeof(FThreadBuffer);
+		memcpy(NewArray, GThreadBufferArray, OldSize);
+		Writer_MemoryFree(GThreadBufferArray, OldSize);
+#if TRACE_PRIVATE_STATISTICS
+		AtomicSubRelaxed(&GTraceStatistics.FixedBufferAllocated, OldSize);
+#endif
+	}
+
+	GThreadBufferArray = NewArray;
+	GThreadBufferArrayCapacity = NewArrayCapacity;
+}
+
 ////////////////////////////////////////////////////////////////////////////////
 static bool Writer_DrainBuffer(uint32 ThreadId, FWriteBuffer* Buffer)
 {
@@ -162,6 +199,7 @@ static bool Writer_DrainBuffer(uint32 ThreadId, FWriteBuffer* Buffer)
 void Writer_DrainBuffers()
 {
 	FProfilerScope _(__func__);
+
 	struct FRetireList
 	{
 		FWriteBuffer* __restrict Head = nullptr;
@@ -191,7 +229,7 @@ void Writer_DrainBuffers()
 
 		NewThreadCursor = NextThread;
 	}
-	
+
 	FRetireList RetireList;
 
 	FWriteBuffer* __restrict ActiveThreadList = GActiveThreadList;
@@ -199,45 +237,64 @@ void Writer_DrainBuffers()
 
 	// Now we've two lists of known and new threads. Each of these lists in turn is
 	// a list of that thread's buffers (where it is writing trace events to).
+
+	// Step one: Identify the buffers available now, for each available thread.
+	uint32 ThreadCount = 0;
 	for (FWriteBuffer* __restrict Buffer : { ActiveThreadList, NewThreadList })
 	{
 		// For each thread...
 		for (FWriteBuffer* __restrict NextThread; Buffer != nullptr; Buffer = NextThread)
 		{
 			NextThread = Buffer->NextThread;
-			uint32 ThreadId = Buffer->ThreadId;
 
-			// Count how many buffers are available now, for the current thread.
-			uint32 DrainBufferLimit = 1;
-			for (FWriteBuffer* __restrict CrtBuffer = Buffer; CrtBuffer != nullptr; ++DrainBufferLimit)
+			if (ThreadCount >= GThreadBufferArrayCapacity)
 			{
-				CrtBuffer = AtomicLoadAcquire(&(CrtBuffer->NextBuffer));
+				uint32 NewArrayCapacity = (GThreadBufferArrayCapacity == 0) ? 256 : GThreadBufferArrayCapacity << 1;
+				Writer_ReallocateThreadBufferArray(NewArrayCapacity);
 			}
 
-			// For each of the thread's buffers...
-			for (FWriteBuffer* __restrict NextBuffer; Buffer != nullptr; Buffer = NextBuffer)
+			// Count how many buffers are available now, for the current thread.
+			uint32 NumBuffers = 1;
+			for (FWriteBuffer* __restrict CrtBuffer = Buffer; CrtBuffer != nullptr; ++NumBuffers)
 			{
-				if (--DrainBufferLimit == 0)
-				{
-					// Ignore buffers added while we drain buffers for this thread.
-					break;
-				}
+				CrtBuffer = AtomicLoadAcquire(&(CrtBuffer->NextBuffer));
+			}
+			GThreadBufferArray[ThreadCount].FirstBuffer = Buffer;
+			GThreadBufferArray[ThreadCount].DrainBufferLimit = NumBuffers;
+			++ThreadCount;
+		}
+	}
 
-				if (Writer_DrainBuffer(ThreadId, Buffer))
-				{
-					break;
-				}
+	// Step two: Drain only the threads and their buffers identified in step one.
+	for (uint32 ThreadIndex = 0; ThreadIndex < ThreadCount; ++ThreadIndex)
+	{
+		FWriteBuffer* __restrict Buffer = GThreadBufferArray[ThreadIndex].FirstBuffer;
+		uint32 ThreadId = Buffer->ThreadId;
+		uint32 DrainBufferLimit = GThreadBufferArray[ThreadIndex].DrainBufferLimit;
 
-				// Retire the buffer
-				NextBuffer = AtomicLoadAcquire(&(Buffer->NextBuffer));
-				RetireList.Insert(Buffer);
+		// For each of the thread's buffers...
+		for (FWriteBuffer* __restrict NextBuffer; Buffer != nullptr; Buffer = NextBuffer)
+		{
+			if (--DrainBufferLimit == 0)
+			{
+				// Ignore buffers added after the "step one" above.
+				break;
 			}
 
-			if (Buffer != nullptr)
+			if (Writer_DrainBuffer(ThreadId, Buffer))
 			{
-				Buffer->NextThread = GActiveThreadList;
-				GActiveThreadList = Buffer;
+				break;
 			}
+
+			// Retire the buffer.
+			NextBuffer = AtomicLoadAcquire(&(Buffer->NextBuffer));
+			RetireList.Insert(Buffer);
+		}
+
+		if (Buffer != nullptr)
+		{
+			Buffer->NextThread = GActiveThreadList;
+			GActiveThreadList = Buffer;
 		}
 	}
 
