TraceLog: Changed the trace worker thread sleep duration (configured using ?SleepTimeInMS? or ?SleepTimeWhenMemoryTracingInMS? ini file settings or using the -tracethreadsleeptime=<ms> command line option) to be a time limit for the entire update loop of the worker thread (instead of a fixed sleep time between two updates).

diff --git a/Engine/Source/Runtime/TraceLog/Private/Trace/Writer.cpp b/Engine/Source/Runtime/TraceLog/Private/Trace/Writer.cpp
--- a/Engine/Source/Runtime/TraceLog/Private/Trace/Writer.cpp
+++ b/Engine/Source/Runtime/TraceLog/Private/Trace/Writer.cpp
@@ -887,11 +887,18 @@ static void Writer_WorkerThread()
 	// Enable the pool limit when the worker thread is running
 	Writer_SetBlockPoolLimit(GBlockPoolMaxSize);
 
+	const double TicksToMS = 1000.0 / (double)TimeGetFrequency();
+
+	uint64 PrevTime = TimeGetTimestamp();
 	while (!AtomicLoadRelaxed(&GWorkerThreadQuit))
 	{
 		Writer_WorkerUpdate();
 
-		ThreadSleep(GSleepTimeInMS);
+		uint64 Time = TimeGetTimestamp();
+		uint32 DurationMS = uint32(double(Time - PrevTime) * TicksToMS);
+		PrevTime = Time;
+
+		ThreadSleep(DurationMS < GSleepTimeInMS ? GSleepTimeInMS - DurationMS : 0);
 	}
 
 	// Reset the limit as no one will pick up data
