[Trace]
- Fix GStatistics benign race

diff --git a/Engine/Source/Runtime/TraceLog/Private/Trace/TlsBuffer.cpp b/Engine/Source/Runtime/TraceLog/Private/Trace/TlsBuffer.cpp
--- a/Engine/Source/Runtime/TraceLog/Private/Trace/TlsBuffer.cpp
+++ b/Engine/Source/Runtime/TraceLog/Private/Trace/TlsBuffer.cpp
@@ -183,7 +183,7 @@ static bool Writer_DrainBuffer(uint32 ThreadId, FWriteBuffer* Buffer)
 	if (uint32 SizeToReap = uint32(Committed - Buffer->Reaped))
 	{
 #if TRACE_PRIVATE_STATISTICS
-		GTraceStatistics.BytesTraced += SizeToReap;
+		FPlatformAtomics::AtomicStore_Relaxed((int64*)&GTraceStatistics.BytesTraced, FPlatformAtomics::AtomicRead_Relaxed((int64*)&GTraceStatistics.BytesTraced) + SizeToReap);
 #endif
 
 		Writer_TailAppend(ThreadId, Buffer->Reaped, SizeToReap);
diff --git a/Engine/Source/Runtime/TraceLog/Private/Trace/Writer.cpp b/Engine/Source/Runtime/TraceLog/Private/Trace/Writer.cpp
--- a/Engine/Source/Runtime/TraceLog/Private/Trace/Writer.cpp
+++ b/Engine/Source/Runtime/TraceLog/Private/Trace/Writer.cpp
@@ -354,7 +354,7 @@ static bool Writer_FlushSendBuffer() { return true; }
 static void Writer_SendDataImplNoInstr(const void* Data, uint32 Size)
 {
 #if TRACE_PRIVATE_STATISTICS
-	GTraceStatistics.BytesSent += Size;
+	FPlatformAtomics::AtomicStore_Relaxed((int64*)&GTraceStatistics.BytesSent, FPlatformAtomics::AtomicRead_Relaxed((int64*)&GTraceStatistics.BytesSent) + Size);
 #endif
 
 #if TRACE_PRIVATE_BUFFER_SEND
