Added tracking for free block in BlockPool

diff --git a/Engine/Source/Runtime/TraceLog/Private/Trace/BlockPool.cpp b/Engine/Source/Runtime/TraceLog/Private/Trace/BlockPool.cpp
--- a/Engine/Source/Runtime/TraceLog/Private/Trace/BlockPool.cpp
+++ b/Engine/Source/Runtime/TraceLog/Private/Trace/BlockPool.cpp
@@ -67,7 +67,7 @@ void Writer_UnsetBlockPoolLimit()
 }
 
 ////////////////////////////////////////////////////////////////////////////////
-static FPoolBlockList Writer_AllocateBlockList(uint32 PageSize)
+static FPoolBlockList Writer_AllocateBlockList(uint32 PageSize, uint32& OutBlockCount)
 {
 	// The free list is empty so we have to populate it with some new blocks.
 	uint8* PageBase = (uint8*)Writer_MemoryAllocate(PageSize, PLATFORM_CACHE_LINE_SIZE);
@@ -84,7 +84,8 @@ static FPoolBlockList Writer_AllocateBlockList(uint32 PageSize)
 	// Link subsequent blocks together
 	uint8* FirstBlock = PageBase + GPoolBlockSize - sizeof(FWriteBuffer);
 	uint8* Block = FirstBlock;
-	for (int i = 1, n = PageSize / GPoolBlockSize; ; ++i)
+	const uint32 BlockCount = PageSize / GPoolBlockSize;
+	for (int i = 1, n = BlockCount; ; ++i)
 	{
 		auto* Buffer = (FWriteBuffer*)Block;
 		Buffer->Size = uint16(BufferSize);
@@ -97,6 +98,9 @@ static FPoolBlockList Writer_AllocateBlockList(uint32 PageSize)
 		Block += GPoolBlockSize;
 	}
 
+	AtomicAddRelaxed(&GTraceStatistics.BlockPoolAllocatedBlocks, BlockCount);
+	OutBlockCount = BlockCount;
+
 	// Keep track of allocation base so we can free it on shutdown
 	FWriteBuffer* NextBuffer = (FWriteBuffer*)FirstBlock;
 	NextBuffer->Size -= sizeof(FPoolPage);
@@ -149,6 +153,7 @@ FWriteBuffer* Writer_AllocateBlockFromPool()
 		// If we didn't fetch the sentinal then we've taken a block we can use
 		if (Owned != nullptr)
 		{
+			AtomicSubRelaxed(&GTraceStatistics.BlockPoolFreeBlocks, 1u);
 			Ret = (FWriteBuffer*)Owned;
 			break;
 		}
@@ -171,7 +176,8 @@ FWriteBuffer* Writer_AllocateBlockFromPool()
 			continue;
 		}
 
-		FPoolBlockList BlockList = Writer_AllocateBlockList(GPoolPageSize);
+		uint32 BlockCount = 0;
+		FPoolBlockList BlockList = Writer_AllocateBlockList(GPoolPageSize, BlockCount);
 		Ret = BlockList.Head;
 
 		// And insert the block list into the freelist. 'Block' is now the last block
@@ -182,6 +188,8 @@ FWriteBuffer* Writer_AllocateBlockFromPool()
 
 			if (AtomicCompareExchangeRelease(&GPoolFreeList, Ret->NextBuffer, FreeListValue))
 			{
+				AtomicAddRelaxed(&GTraceStatistics.BlockPoolFreeBlocks, BlockCount - 1);
+				
 				break;
 			}
 		}
diff --git a/Engine/Source/Runtime/TraceLog/Private/Trace/TlsBuffer.cpp b/Engine/Source/Runtime/TraceLog/Private/Trace/TlsBuffer.cpp
--- a/Engine/Source/Runtime/TraceLog/Private/Trace/TlsBuffer.cpp
+++ b/Engine/Source/Runtime/TraceLog/Private/Trace/TlsBuffer.cpp
@@ -231,6 +231,7 @@ void Writer_DrainBuffers()
 	}
 
 	FRetireList RetireList;
+	uint32 RetiredCount = 0;
 
 	FWriteBuffer* __restrict ActiveThreadList = GActiveThreadList;
 	GActiveThreadList = nullptr;
@@ -289,6 +290,7 @@ void Writer_DrainBuffers()
 			// Retire the buffer.
 			NextBuffer = AtomicLoadAcquire(&(Buffer->NextBuffer));
 			RetireList.Insert(Buffer);
+			++RetiredCount;
 		}
 
 		if (Buffer != nullptr)
@@ -302,6 +304,7 @@ void Writer_DrainBuffers()
 	if (RetireList.Head != nullptr)
 	{
 		Writer_FreeBlockListToPool(RetireList.Head, RetireList.Tail);
+		AtomicAddRelaxed(&GTraceStatistics.BlockPoolFreeBlocks, RetiredCount);
 	}
 }
 
@@ -329,6 +332,7 @@ void Writer_DrainLocalBuffers()
 	};
 
 	FRetireList RetireList;
+	uint32 RetiredCount = 0;
 
 	FWriteBuffer* __restrict ActiveThreadList = GActiveThreadList;
 	GActiveThreadList = nullptr;
@@ -356,6 +360,7 @@ void Writer_DrainLocalBuffers()
 					// Retire the buffer
 					NextBuffer = AtomicLoadAcquire(&(Buffer->NextBuffer));
 					RetireList.Insert(Buffer);
+					++RetiredCount;
 				}
 			}
 
@@ -371,6 +376,7 @@ void Writer_DrainLocalBuffers()
 	if (RetireList.Head != nullptr)
 	{
 		Writer_FreeBlockListToPool(RetireList.Head, RetireList.Tail);
+		AtomicAddRelaxed(&GTraceStatistics.BlockPoolFreeBlocks, RetiredCount);
 	}
 }
 
diff --git a/Engine/Source/Runtime/TraceLog/Private/Trace/Trace.cpp b/Engine/Source/Runtime/TraceLog/Private/Trace/Trace.cpp
--- a/Engine/Source/Runtime/TraceLog/Private/Trace/Trace.cpp
+++ b/Engine/Source/Runtime/TraceLog/Private/Trace/Trace.cpp
@@ -1,6 +1,7 @@
 // Copyright Epic Games, Inc. All Rights Reserved.
 
 #include "Trace/Trace.inl" // should be Config.h :(
+#include "Trace/Detail/Atomic.h"
 
 #if TRACE_PRIVATE_MINIMAL_ENABLED
 
@@ -125,6 +126,8 @@ void GetStatistics(FStatistics& Out)
 	Out.BytesEmitted = Private::AtomicLoadRelaxed(&Private::GTraceStatistics.BytesEmitted);
 	Out.MemoryUsed = Private::AtomicLoadRelaxed(&Private::GTraceStatistics.MemoryUsed);
 	Out.BlockPoolAllocated = Private::AtomicLoadRelaxed(&Private::GTraceStatistics.BlockPoolAllocated);
+	Out.BlockPoolFreeBlocks = Private::AtomicLoadRelaxed(&Private::GTraceStatistics.BlockPoolFreeBlocks);
+	Out.BlockPoolAllocatedBlocks = Private::AtomicLoadRelaxed(&Private::GTraceStatistics.BlockPoolAllocatedBlocks);
 	Out.SharedBufferAllocated = Private::AtomicLoadRelaxed(&Private::GTraceStatistics.SharedBufferAllocated);
 	Out.FixedBufferAllocated = Private::AtomicLoadRelaxed(&Private::GTraceStatistics.FixedBufferAllocated);
 	Out.CacheAllocated = Private::AtomicLoadRelaxed(&Private::GTraceStatistics.CacheAllocated);
diff --git a/Engine/Source/Runtime/TraceLog/Public/Trace/Trace.h b/Engine/Source/Runtime/TraceLog/Public/Trace/Trace.h
--- a/Engine/Source/Runtime/TraceLog/Public/Trace/Trace.h
+++ b/Engine/Source/Runtime/TraceLog/Public/Trace/Trace.h
@@ -218,6 +218,8 @@ struct FStatistics
 	uint64	BytesEmitted			= 0;	// Bytes emitted but potentially not yet written
 	uint64	MemoryUsed				= 0;	// Memory allocated by TraceLog allocator functions
 	uint64	BlockPoolAllocated		= 0;	// Memory allocated for the (TLS) block pool.
+	uint32  BlockPoolFreeBlocks		= 0;	// Number of available blocks
+	uint32	BlockPoolAllocatedBlocks= 0;	// Number of allocated blocks
 	uint32	SharedBufferAllocated	= 0;	// Memory allocated for shared buffers
 	uint32	FixedBufferAllocated	= 0;	// Memory allocated for fixed buffers (tail, send)
 	uint32	CacheAllocated			= 0;	// Total memory allocated in cache buffers
