[Trace] Fixed the time limitation for the update loop of the trace worker thread.

diff --git a/Engine/Source/Runtime/TraceLog/Private/Trace/Writer.cpp b/Engine/Source/Runtime/TraceLog/Private/Trace/Writer.cpp
--- a/Engine/Source/Runtime/TraceLog/Private/Trace/Writer.cpp
+++ b/Engine/Source/Runtime/TraceLog/Private/Trace/Writer.cpp
@@ -891,14 +891,12 @@ static void Writer_WorkerThread()
 
 	const double TicksToMS = 1000.0 / (double)TimeGetFrequency();
 
-	uint64 PrevTime = TimeGetTimestamp();
 	while (!AtomicLoadRelaxed(&GWorkerThreadQuit))
 	{
+		uint64 StartTime = TimeGetTimestamp();
 		Writer_WorkerUpdate();
-
-		uint64 Time = TimeGetTimestamp();
-		uint32 DurationMS = uint32(double(Time - PrevTime) * TicksToMS);
-		PrevTime = Time;
+		uint64 EndTime = TimeGetTimestamp();
+		uint32 DurationMS = uint32(double(EndTime - StartTime) * TicksToMS);
 
 		ThreadSleep(DurationMS < GSleepTimeInMS ? GSleepTimeInMS - DurationMS : 0);
 	}
