Make default argument null to be backwards compatible.

#rnx

diff --git a/Engine/Source/Runtime/TraceLog/Public/Trace/Detail/Channel.h b/Engine/Source/Runtime/TraceLog/Public/Trace/Detail/Channel.h
--- a/Engine/Source/Runtime/TraceLog/Public/Trace/Detail/Channel.h
+++ b/Engine/Source/Runtime/TraceLog/Public/Trace/Detail/Channel.h
@@ -81,7 +81,7 @@ public:
 	 * @param OutReason Optional string pointer. Will contain deny reason on failure.
 	 * @return New channel state
 	 */
-	TRACELOG_API bool	Toggle(bool bEnabled, const TCHAR** OutReason);
+	TRACELOG_API bool	Toggle(bool bEnabled, const TCHAR** OutReason = nullptr);
 
 	static void Initialize();
 	static Iter	ReadNew();
