Add custom channel macros to minimal trace api

#rnx

diff --git a/Engine/Source/Runtime/TraceLog/Public/Trace/Trace.h b/Engine/Source/Runtime/TraceLog/Public/Trace/Trace.h
--- a/Engine/Source/Runtime/TraceLog/Public/Trace/Trace.h
+++ b/Engine/Source/Runtime/TraceLog/Public/Trace/Trace.h
@@ -471,6 +471,9 @@ UE_TRACE_API FChannel* FindChannel(FChannelId ChannelId) UE_TRACE_IMPL(nullptr);
 #define UE_TRACE_MINIMAL_GET_DEFINITION_TYPE_ID(LoggerName, EventName)									TRACE_IMPL_MINIMAL(GET_DEFINITION_TYPE_ID, LoggerName, EventName)
 #define UE_TRACE_MINIMAL_LOG_DEFINITION(LoggerName, EventName, Id, ChannelsExpr, ...)					TRACE_IMPL_MINIMAL(LOG_DEFINITION, LoggerName, EventName, Id, ChannelsExpr, ##__VA_ARGS__)
 #define UE_TRACE_MINIMAL_CHANNEL(ChannelName, ...)														TRACE_IMPL_MINIMAL(CHANNEL, ChannelName, ##__VA_ARGS__)
+#define UE_TRACE_MINIMAL_CHANNEL_CUSTOM(ChannelName, ChannelClass, ...)									TRACE_IMPL_MINIMAL(CHANNEL_CUSTOM, ChannelName, ChannelClass, ##__VA_ARGS__)
 #define UE_TRACE_MINIMAL_CHANNEL_EXTERN(ChannelName, ...)												TRACE_IMPL_MINIMAL(CHANNEL_EXTERN, ChannelName, ##__VA_ARGS__)
+#define UE_TRACE_MINIMAL_CHANNEL_CUSTOM_EXTERN(ChannelName, ChannelClass, ...)							TRACE_IMPL_MINIMAL(CHANNEL_CUSTOM_EXTERN, ChannelName, ChannelClass, ##__VA_ARGS__)
 #define UE_TRACE_MINIMAL_CHANNEL_DEFINE(ChannelName, ...)												TRACE_IMPL_MINIMAL(CHANNEL_DEFINE, ChannelName, ##__VA_ARGS__)
+#define UE_TRACE_MINIMAL_CHANNEL_CUSTOM_DEFINE(ChannelName, ChannelClass, ...)							TRACE_IMPL_MINIMAL(CHANNEL_CUSTOM_DEFINE, ChannelName, ChannelClass, ##__VA_ARGS__)
 #define UE_TRACE_MINIMAL_CHANNELEXPR_IS_ENABLED(ChannelsExpr)											TRACE_IMPL_MINIMAL(CHANNELEXPR_IS_ENABLED, ChannelsExpr)
