TraceLog: Update auto* to FTidPacket*

#rnx

diff --git a/Engine/Source/Runtime/TraceLog/Private/Trace/Tail.cpp b/Engine/Source/Runtime/TraceLog/Private/Trace/Tail.cpp
--- a/Engine/Source/Runtime/TraceLog/Private/Trace/Tail.cpp
+++ b/Engine/Source/Runtime/TraceLog/Private/Trace/Tail.cpp
@@ -277,7 +277,7 @@ void Writer_TailAppend(uint32 ThreadId, uint8* __restrict Data, uint32 Size)
 		}
 	}
 
-	auto* Packet = GPacketRing.Append<FTidPacket>(Size);
+	FTidPacket* Packet = GPacketRing.Append<FTidPacket>(Size);
 	Packet->ThreadId = uint16(ThreadId);
 #if UE_TRACE_PACKET_VERIFICATION
 	Packet->ThreadId |= FTidPacketBase::Verification;
