TraceLog: null pointer check to fix nullability.NullPassedToNonnull

#rnx

diff --git a/Engine/Source/Runtime/TraceLog/Private/Trace/Tail.cpp b/Engine/Source/Runtime/TraceLog/Private/Trace/Tail.cpp
--- a/Engine/Source/Runtime/TraceLog/Private/Trace/Tail.cpp
+++ b/Engine/Source/Runtime/TraceLog/Private/Trace/Tail.cpp
@@ -282,7 +282,10 @@ void Writer_TailAppend(uint32 ThreadId, uint8* __restrict Data, uint32 Size)
 #if UE_TRACE_PACKET_VERIFICATION
 	Packet->ThreadId |= FTidPacketBase::Verification;
 #endif
-	::memcpy(Packet->Data, Data, Size);
+	if (Packet->Data != nullptr)
+	{
+		::memcpy(Packet->Data, Data, Size);
+	}
 
 	Writer_SendDataRaw(Packet, Packet->PacketSize);
 #if UE_TRACE_PACKET_VERIFICATION
