Fix a startup deadlock when launching with -trace=memory without  tracehost or tracefile. There was a race condition setting a buffer size on the worker thread while the main thread was pushing updates.

diff --git a/Engine/Source/Runtime/TraceLog/Private/Trace/Writer.cpp b/Engine/Source/Runtime/TraceLog/Private/Trace/Writer.cpp
--- a/Engine/Source/Runtime/TraceLog/Private/Trace/Writer.cpp
+++ b/Engine/Source/Runtime/TraceLog/Private/Trace/Writer.cpp
@@ -886,9 +886,14 @@ static void Writer_WorkerThread()
 {
 	ThreadRegister(TEXT("Trace"), 0, INT_MAX);
 
+	// Don't set GBlockPoolMaxSize until the spawning thread is fully done
+	while (!AtomicCompareExchangeAcquire(&GUpdateInProgress, 1u, 0u)) {}
+
 	// Enable the pool limit when the worker thread is running
 	Writer_SetBlockPoolLimit(GBlockPoolMaxSize);
 
+	AtomicExchangeRelease(&GUpdateInProgress, 0u);
+
 	const double TicksToMS = 1000.0 / (double)TimeGetFrequency();
 
 	while (!AtomicLoadRelaxed(&GWorkerThreadQuit))
