[Insights] * TraceLog: Fixed "SentTo" control command to be able to specify the port number after the host name. The default port 1981 is still used if no port is specified.
* Frontend/Connection: It is now possible to specify port number after the host ip/name, in the Trace Destination text box.

diff --git a/Engine/Source/Runtime/TraceLog/Private/Trace/Control.cpp b/Engine/Source/Runtime/TraceLog/Private/Trace/Control.cpp
--- a/Engine/Source/Runtime/TraceLog/Private/Trace/Control.cpp
+++ b/Engine/Source/Runtime/TraceLog/Private/Trace/Control.cpp
@@ -290,7 +290,27 @@ void Writer_InitializeControl()
 		{
 			if (ArgC > 0)
 			{
-				Writer_SendTo(ArgV[0]);
+				const ANSICHAR* Host = ArgV[0];
+				uint32 Port = 0;
+				uint32 Flags = 0;
+
+				constexpr SIZE_T HostNameBufferSize = 256;
+				ANSICHAR HostNameBuffer[HostNameBufferSize] = {};
+				const ANSICHAR* PortStr = FPlatformString::Strchr(Host, ':');
+				if (PortStr)
+				{
+					++PortStr;
+					SIZE_T MaxLen = SIZE_T(PTRINT(PortStr - Host));
+					if (MaxLen > HostNameBufferSize)
+					{
+						MaxLen = HostNameBufferSize;
+					}
+					FPlatformString::Strncpy(HostNameBuffer, ArgV[0], MaxLen);
+					Host = HostNameBuffer;
+					Port = uint32(FPlatformString::Atoi(PortStr));
+				}
+
+				Writer_SendTo(Host, Flags, Port);
 			}
 		}
 	);
