Make Trace FChannels Next pointer thread safe

The Next is accessed in the following ways by traversing channel lists:
Whenever we enable channels by name (from any thread) we iterate main/new channels lists to find a matching channel object. (FChannel::FindChannel)
Every update loop the worker thread iterates new channels list, announces them and moves them to the main list. Part of that updates the Next pointer for channels (FChannel::Iter::~Iter).

#rnx

diff --git a/Engine/Source/Runtime/TraceLog/Private/Trace/Channel.cpp b/Engine/Source/Runtime/TraceLog/Private/Trace/Channel.cpp
--- a/Engine/Source/Runtime/TraceLog/Private/Trace/Channel.cpp
+++ b/Engine/Source/Runtime/TraceLog/Private/Trace/Channel.cpp
@@ -92,8 +92,9 @@ FChannel::Iter::~Iter()
 	using namespace Private;
 	for (auto* Node = (FChannel*)Inner[2];; PlatformYield())
 	{
-		Node->Next = AtomicLoadRelaxed(&GHeadChannel);
-		if (AtomicCompareExchangeRelaxed(&GHeadChannel, (FChannel*)Inner[1], Node->Next))
+		FChannel* OldHead = AtomicLoadRelaxed(&GHeadChannel);
+		AtomicStoreRelaxed(&(Node->Next), OldHead);
+		if (AtomicCompareExchangeRelaxed(&GHeadChannel, (FChannel*)Inner[1], OldHead))
 		{
 			break;
 		}
@@ -103,10 +104,12 @@ FChannel::Iter::~Iter()
 ///////////////////////////////////////////////////////////////////////////////
 const FChannel* FChannel::Iter::GetNext()
 {
-	auto* Ret = (const FChannel*)Inner[0];
+	using namespace Private;
+
+	auto* Ret = (FChannel*)Inner[0];
 	if (Ret != nullptr)
 	{
-		Inner[0] = Ret->Next;
+		Inner[0] = AtomicLoadRelaxed(&Ret->Next);
 		if (Inner[0] != nullptr)
 		{
 			Inner[2] = Inner[0];
@@ -151,7 +154,7 @@ void FChannel::Setup(const ANSICHAR* InChannelName, const InitArgs& InArgs)
 	for (;; PlatformYield())
 	{
 		FChannel* HeadChannel = AtomicLoadRelaxed(&GNewChannelList);
-		Next = HeadChannel;
+		AtomicStoreRelaxed(&Next, HeadChannel);
 		if (AtomicCompareExchangeRelease(&GNewChannelList, this, Next))
 		{
 			break;
@@ -206,7 +209,7 @@ void FChannel::ToggleAllInternal(bool bEnabled, bool bIssueCallback)
 	};
 	for (FChannel* Channel : ChannelLists)
 	{
-		for (; Channel != nullptr; Channel = (FChannel*)(Channel->Next))
+		for (; Channel != nullptr; Channel = AtomicLoadRelaxed(&Channel->Next))
 		{
 			Channel->ToggleInternal(bEnabled, bIssueCallback);
 		}
@@ -225,7 +228,7 @@ void FChannel::PanicDisableAll()
 	};
 	for (FChannel* Channel : ChannelLists)
 	{
-		for (; Channel != nullptr; Channel = (FChannel*)(Channel->Next))
+		for (; Channel != nullptr; Channel = AtomicLoadRelaxed(&Channel->Next))
 		{
 			AtomicStoreRelaxed(&Channel->Enabled, -1);
 		}
@@ -247,7 +250,7 @@ FChannel* FChannel::FindChannel(const ANSICHAR* ChannelName)
 	};
 	for (FChannel* Channel : ChannelLists)
 	{
-		for (; Channel != nullptr; Channel = Channel->Next)
+		for (; Channel != nullptr; Channel = AtomicLoadRelaxed(&Channel->Next))
 		{
 			if (Channel->Name.Hash == ChannelNameHash)
 			{
@@ -271,7 +274,7 @@ FChannel* FChannel::FindChannel(FChannelId ChannelId)
 	};
 	for (FChannel* Channel : ChannelLists)
 	{
-		for (; Channel != nullptr; Channel = Channel->Next)
+		for (; Channel != nullptr; Channel = AtomicLoadRelaxed(&Channel->Next))
 		{
 			if (Channel->Name.Hash == ChannelId)
 			{
@@ -296,7 +299,7 @@ void FChannel::EnumerateChannels(ChannelIterCallback Func, void* User)
 	FChannelInfo Info;
 	for (FChannel* Channel : ChannelLists)
 	{
-		for (; Channel != nullptr; Channel = Channel->Next)
+		for (; Channel != nullptr; Channel = AtomicLoadRelaxed(&Channel->Next))
 		{
 			Info.Name = Channel->Name.Ptr;
 			Info.Desc = Channel->Args.Desc;
diff --git a/Engine/Source/Runtime/TraceLog/Public/Trace/Detail/Channel.h b/Engine/Source/Runtime/TraceLog/Public/Trace/Detail/Channel.h
--- a/Engine/Source/Runtime/TraceLog/Public/Trace/Detail/Channel.h
+++ b/Engine/Source/Runtime/TraceLog/Public/Trace/Detail/Channel.h
@@ -102,7 +102,7 @@ private:
 	bool ToggleInternal(bool bEnabled, bool bIssueCallback, const TCHAR** OutReason = nullptr);
 	static void ToggleAllInternal(bool bEnabled, bool bIssueCallback);
 	
-	FChannel* Next;
+	FChannel* volatile Next;
 	struct
 	{
 		const ANSICHAR* Ptr;
