[TraceLog]
* Made Protocol.h iwyu friendly by adding export pragmas around all includes

diff --git a/Engine/Source/Runtime/TraceLog/Public/Trace/Detail/Protocol.h b/Engine/Source/Runtime/TraceLog/Public/Trace/Detail/Protocol.h
--- a/Engine/Source/Runtime/TraceLog/Public/Trace/Detail/Protocol.h
+++ b/Engine/Source/Runtime/TraceLog/Public/Trace/Detail/Protocol.h
@@ -10,6 +10,7 @@
 	#pragma warning(disable : 4200) // non-standard zero-sized array
 #endif
 
+// IWYU pragma: begin_exports
 #include "Protocols/Protocol0.h"
 #include "Protocols/Protocol1.h"
 #include "Protocols/Protocol2.h"
@@ -18,6 +19,7 @@
 #include "Protocols/Protocol5.h"
 #include "Protocols/Protocol6.h"
 #include "Protocols/Protocol7.h"
+// IWYU pragma: end_exports
 
 #if defined(_MSC_VER)
 	#pragma warning(pop)
