Trace Panic when block pool is full.

Add ability for TraceLog to put in a "panic" mode, where the worker thread is immediately stopped, further updates prevented, and tracing will be disabled for the remainder of the process lifetime. This panic mode is triggered when the block pool cannot prodivde a free block in a set time. The block pool then returns a special "panic block" shared across all threads, which appears as a normal block, but will never increase it's buffer cursor. All threads can then procede writing data as normal.

diff --git a/Engine/Source/Runtime/TraceLog/Private/Trace/BlockPool.cpp b/Engine/Source/Runtime/TraceLog/Private/Trace/BlockPool.cpp
--- a/Engine/Source/Runtime/TraceLog/Private/Trace/BlockPool.cpp
+++ b/Engine/Source/Runtime/TraceLog/Private/Trace/BlockPool.cpp
@@ -46,6 +46,7 @@ static uint32							GPoolMaxSize		= ~0; // Starts unlimited, written by initiali
 T_ALIGN static FWriteBuffer* volatile	GPoolFreeList;		// = nullptr;
 T_ALIGN static UPTRINT volatile			GPoolFutex;			// = 0
 T_ALIGN static FPoolPage* volatile		GPoolPageList;		// = nullptr;
+static const uint32						GPoolLimitTimeoutMs	= 4 * 1000;
 #undef T_ALIGN
 
 ////////////////////////////////////////////////////////////////////////////////
@@ -89,6 +90,7 @@ static FPoolBlockList Writer_AllocateBlockList(uint32 PageSize, uint32& OutBlock
 	{
 		auto* Buffer = (FWriteBuffer*)Block;
 		Buffer->Size = uint16(BufferSize);
+		Buffer->CursorMaskInv = 0;
 		if (i >= n)
 		{
 			break;
@@ -112,6 +114,34 @@ static FPoolBlockList Writer_AllocateBlockList(uint32 PageSize, uint32& OutBlock
 	return { NextBuffer, (FWriteBuffer*)Block };
 }
 
+////////////////////////////////////////////////////////////////////////////////
+FWriteBuffer* Writer_GetPanicPool()
+{
+	static FWriteBuffer* SPanicBuffer = []()
+	{
+		uint8 *Allocation = (uint8 *)Writer_MemoryAllocate(
+			GPoolBlockSize, PLATFORM_CACHE_LINE_SIZE);
+#if TRACE_PRIVATE_STATISTICS
+		AtomicAddRelaxed(&GTraceStatistics.BlockPoolAllocated, uint64(GPoolBlockSize));
+#endif
+		uint32 BufferSize = GPoolBlockSize - sizeof(FWriteBuffer);
+		FWriteBuffer* PanicBuffer = (FWriteBuffer *)(Allocation + BufferSize);
+		static_assert(GPoolBlockSize < UINT16_MAX);
+		PanicBuffer->CursorMaskInv = 0;
+		PanicBuffer->Size = uint16(BufferSize);
+		PanicBuffer->NextBuffer = nullptr;
+		PanicBuffer->NextThread = nullptr;
+		PanicBuffer->SetCursor(Allocation);
+		PanicBuffer->Reaped = Allocation;
+		PanicBuffer->EtxOffset = 0 - int32(sizeof(FWriteBuffer));
+		// This mask will cause the cursor to stay at the start of
+		// the buffer, without ever increasing
+		PanicBuffer->CursorMaskInv = ~0;
+		return PanicBuffer;
+	}();
+	return SPanicBuffer;
+}
+
 ////////////////////////////////////////////////////////////////////////////////
 FWriteBuffer* Writer_AllocateBlockFromPool()
 {
@@ -160,9 +190,20 @@ FWriteBuffer* Writer_AllocateBlockFromPool()
 
 		// If we have hit the pool limit, continue looping until the worker
 		// thread has had time write out and give back free buffers.
-		if (AtomicLoadAcquire(&GPoolTotalAllocated) >= GPoolMaxSize)
+		
+		if ((AtomicLoadAcquire(&GPoolTotalAllocated) >= GPoolMaxSize))
 		{
-			StallStart = StallStart ? StallStart : TimeGetRelativeTimestamp();
+			const uint64 Ts = TimeGetRelativeTimestamp();
+			StallStart = StallStart ? StallStart : Ts;
+			if ((Ts - StallStart) > GPoolLimitTimeoutMs)
+			{
+				// At this point we have timed out waiting for a free block.
+				// This likely means the worker thread cannot keep up with
+				// tracing bandwidth or that I/O is stalled.
+				Panic();
+				Ret = Writer_GetPanicPool();
+				break;
+			}
 			continue;
 		}
 
diff --git a/Engine/Source/Runtime/TraceLog/Private/Trace/Field.cpp b/Engine/Source/Runtime/TraceLog/Private/Trace/Field.cpp
--- a/Engine/Source/Runtime/TraceLog/Private/Trace/Field.cpp
+++ b/Engine/Source/Runtime/TraceLog/Private/Trace/Field.cpp
@@ -32,20 +32,20 @@ static void Field_WriteAuxData(uint32 Index, int32 Size, CallbackType&& Callback
 	// along. However, if we're writing a second aux field where the first fetched
 	// a new buffer then it is a different buffer to the one the event's written
 	// into, thus it is us that needs to update Committed.
-	bool bCommit = (Buffer->Cursor == Buffer->Committed);
+	bool bCommit = (Buffer->ReadCursor() == Buffer->Committed);
 
 	// Do we have enough space to write at least [header][overflow-1] bytes of payload?
-	int32 Remaining = int32(ptrdiff_t((uint8*)Buffer - Buffer->Cursor));
+	int32 Remaining = int32(ptrdiff_t((uint8*)Buffer - Buffer->ReadCursor()));
 	auto NextBuffer = [&Buffer, &Remaining, &bCommit] ()
 	{
 		if (bCommit)
 		{
-			AtomicStoreRelease(&(uint8* volatile&)(Buffer->Committed), Buffer->Cursor);
+			AtomicStoreRelease(&(uint8* volatile&)(Buffer->Committed), Buffer->ReadCursor());
 		}
 		bCommit = true;
 
 		Buffer = Writer_NextBuffer();
-		Remaining = int32(ptrdiff_t((uint8*)Buffer - Buffer->Cursor));
+		Remaining = int32(ptrdiff_t((uint8*)Buffer - Buffer->ReadCursor()));
 	};
 
 	if (Remaining <= 0)
@@ -65,13 +65,13 @@ static void Field_WriteAuxData(uint32 Index, int32 Size, CallbackType&& Callback
 		// Write header
 		uint32 Pack = SegmentSize << FAuxHeader::SizeShift;
 		Pack |= Index << FAuxHeader::FieldShift;
-		memcpy(Buffer->Cursor, &Pack, sizeof(uint32)); /* FAuxHeader::Pack */
-		Buffer->Cursor[0] = uint8(EKnownEventUids::AuxData) << EKnownEventUids::_UidShift; /* FAuxHeader::Uid */
-		Buffer->Cursor += sizeof(FAuxHeader);
+		memcpy(Buffer->ReadCursor(), &Pack, sizeof(uint32)); /* FAuxHeader::Pack */
+		*Buffer->ReadCursor() = uint8(EKnownEventUids::AuxData) << EKnownEventUids::_UidShift; /* FAuxHeader::Uid */
+		Buffer->MoveCursor(sizeof(FAuxHeader));
 
 		// Write payload data
-		Callback(Buffer->Cursor, SegmentSize);
-		Buffer->Cursor += SegmentSize;
+		Callback(Buffer->ReadCursor(), SegmentSize);
+		Buffer->MoveCursor((uint16)SegmentSize);
 
 		// Bounds check
 		Size -= SegmentSize;
@@ -85,7 +85,7 @@ static void Field_WriteAuxData(uint32 Index, int32 Size, CallbackType&& Callback
 
 	if (bCommit)
 	{
-		AtomicStoreRelease(&(uint8* volatile&)(Buffer->Committed), Buffer->Cursor);
+		AtomicStoreRelease(&(uint8* volatile&)(Buffer->Committed), Buffer->ReadCursor());
 	}
 }
 
diff --git a/Engine/Source/Runtime/TraceLog/Private/Trace/TlsBuffer.cpp b/Engine/Source/Runtime/TraceLog/Private/Trace/TlsBuffer.cpp
--- a/Engine/Source/Runtime/TraceLog/Private/Trace/TlsBuffer.cpp
+++ b/Engine/Source/Runtime/TraceLog/Private/Trace/TlsBuffer.cpp
@@ -45,7 +45,7 @@ struct FThreadBuffer
 
 ////////////////////////////////////////////////////////////////////////////////
 #define T_ALIGN alignas(PLATFORM_CACHE_LINE_SIZE)
-static FWriteBuffer						GNullWriteBuffer	= { {}, 0, nullptr, nullptr, (uint8*)&GNullWriteBuffer, nullptr, nullptr, 0, 0, 0 };
+static FWriteBuffer						GNullWriteBuffer	= { {}, static_cast<uint16>(~0u), 0, 0, 0, 0, nullptr, nullptr, nullptr, nullptr, (uint8*)&GNullWriteBuffer };
 thread_local FWriteBuffer*				GTlsWriteBuffer		= &GNullWriteBuffer;
 static FWriteBuffer* __restrict			GActiveThreadList;	// = nullptr;
 T_ALIGN static FWriteBuffer* volatile	GNewThreadList;		// = nullptr;
@@ -85,23 +85,23 @@ static FWriteBuffer* Writer_NextBufferInternal(FWriteBuffer* CurrentBuffer)
 #endif
 
 	FWriteBuffer* NextBuffer = Writer_AllocateBlockFromPool();
-	NextBuffer->Cursor = (uint8*)NextBuffer - NextBuffer->Size;
-	NextBuffer->Reaped = NextBuffer->Cursor;
+	NextBuffer->SetCursor((uint8*)NextBuffer - NextBuffer->Size);
+	NextBuffer->Reaped = NextBuffer->ReadCursor();
 	NextBuffer->EtxOffset = 0 - int32(sizeof(FWriteBuffer));
 	AtomicStoreRelaxed(&(NextBuffer->NextBuffer), (FWriteBuffer*)nullptr);
 
 	// Add any capture events into the buffer.
 	if (uint32 RedirectSize = TraceData.GetSize())
 	{
-		memcpy(NextBuffer->Cursor, TraceData.GetData(), RedirectSize);
-		NextBuffer->Cursor += RedirectSize;
+		memcpy(NextBuffer->ReadCursor(), TraceData.GetData(), RedirectSize);
+		NextBuffer->MoveCursor((uint16)RedirectSize);
 	}
 
 	TraceData.Abandon();
 
 	GTlsWriteBuffer = NextBuffer;
 
-	NextBuffer->Committed = NextBuffer->Cursor;
+	NextBuffer->Committed = NextBuffer->ReadCursor();
 
 	if (CurrentBuffer == &GNullWriteBuffer)
 	{
@@ -126,7 +126,7 @@ static FWriteBuffer* Writer_NextBufferInternal(FWriteBuffer* CurrentBuffer)
 	#if TRACE_PRIVATE_STATISTICS
 		// Note that we only track the emitted bytes when a buffer is filled as the overhead
 		// of doing on every event would be too large.
-		const uint64 Used = CurrentBuffer->Cursor - ((uint8*)CurrentBuffer - CurrentBuffer->Size);
+		const uint64 Used = CurrentBuffer->ReadCursor() - ((uint8*)CurrentBuffer - CurrentBuffer->Size);
 		AtomicAddRelaxed(&GTraceStatistics.BytesEmitted, Used);
 	#endif
 
@@ -135,7 +135,7 @@ static FWriteBuffer* Writer_NextBufferInternal(FWriteBuffer* CurrentBuffer)
 		AtomicStoreRelease(&(CurrentBuffer->NextBuffer), NextBuffer);
 
 		// Retire current buffer.
-		int32 EtxOffset = int32(PTRINT((uint8*)(CurrentBuffer) - CurrentBuffer->Cursor));
+		int32 EtxOffset = int32(PTRINT((uint8*)(CurrentBuffer) - CurrentBuffer->ReadCursor()));
 		AtomicStoreRelease(&(CurrentBuffer->EtxOffset), EtxOffset);
 	}
 
@@ -388,7 +388,7 @@ void Writer_EndThreadBuffer()
 		return;
 	}
 
-	const PTRINT CurrentCursor = AtomicLoadRelaxed(GTlsWriteBuffer->Cursor);
+	const PTRINT CurrentCursor = (PTRINT)GTlsWriteBuffer->ReadCursor();
 	int32 EtxOffset = int32(PTRINT((uint8*)GTlsWriteBuffer - CurrentCursor));
 	AtomicStoreRelease(&(GTlsWriteBuffer->EtxOffset), EtxOffset);
 }
diff --git a/Engine/Source/Runtime/TraceLog/Private/Trace/Trace.cpp b/Engine/Source/Runtime/TraceLog/Private/Trace/Trace.cpp
--- a/Engine/Source/Runtime/TraceLog/Private/Trace/Trace.cpp
+++ b/Engine/Source/Runtime/TraceLog/Private/Trace/Trace.cpp
@@ -33,6 +33,7 @@ bool	Writer_IsTracing();
 bool	Writer_IsTracingTo(uint32 (&OutSessionGuid)[4], uint32 (&OutTraceGuid)[4]);
 bool	Writer_Stop();
 uint32	Writer_GetThreadId();
+void	Writer_Panic();
 
 extern FStatistics GTraceStatistics;
 
@@ -109,7 +110,11 @@ void Shutdown()
 ////////////////////////////////////////////////////////////////////////////////
 void Panic()
 {
-	FChannel::PanicDisableAll();
+	static bool bHasPanicked = []() { 
+		FChannel::PanicDisableAll();
+		Private::Writer_Panic();
+		return true;
+	}();
 }
 
 ////////////////////////////////////////////////////////////////////////////////
diff --git a/Engine/Source/Runtime/TraceLog/Private/Trace/WriteBufferRedirect.h b/Engine/Source/Runtime/TraceLog/Private/Trace/WriteBufferRedirect.h
--- a/Engine/Source/Runtime/TraceLog/Private/Trace/WriteBufferRedirect.h
+++ b/Engine/Source/Runtime/TraceLog/Private/Trace/WriteBufferRedirect.h
@@ -98,9 +98,10 @@ inline uint32 TWriteBufferRedirect<BufferSize>::GetCapacity() const
 template <int BufferSize>
 inline void TWriteBufferRedirect<BufferSize>::Reset()
 {
-	Buffer.Cursor = Data + sizeof(uint32);
-	Buffer.Committed = Buffer.Cursor;
-	Buffer.Reaped = Buffer.Cursor;
+	Buffer.CursorMaskInv = 0;
+	Buffer.SetCursor(Data + sizeof(uint32));
+	Buffer.Committed = Buffer.ReadCursor();
+	Buffer.Reaped = Buffer.ReadCursor();
 }
 
 } // namespace Private
diff --git a/Engine/Source/Runtime/TraceLog/Private/Trace/Writer.cpp b/Engine/Source/Runtime/TraceLog/Private/Trace/Writer.cpp
--- a/Engine/Source/Runtime/TraceLog/Private/Trace/Writer.cpp
+++ b/Engine/Source/Runtime/TraceLog/Private/Trace/Writer.cpp
@@ -1447,6 +1447,21 @@ bool Writer_Stop()
 	return true;
 }
 
+////////////////////////////////////////////////////////////////////////////////
+void Writer_Panic()
+{
+	AtomicStoreRelease(&GWorkerThreadQuit, true);
+	if (GWorkerThread && GWorkerThread != Writer_GetThreadId())
+	{
+		ThreadDestroy(GWorkerThread);
+		// Set the update in progress to make sure updates cannot be executed
+		// and snapshots cannot be written.
+		AtomicStoreRelease(&GUpdateInProgress, 1u);
+		GWorkerThread = 0;
+	}
+	UE_TRACE_MESSAGE(PanicError, "Trace has encountered a catastrophic failure. Tracing will be unavailable for the rest of the process lifetime.");
+}
+
 } // namespace Private
 } // namespace Trace
 } // namespace UE
diff --git a/Engine/Source/Runtime/TraceLog/Public/Trace/Detail/LogScope.inl b/Engine/Source/Runtime/TraceLog/Public/Trace/Detail/LogScope.inl
--- a/Engine/Source/Runtime/TraceLog/Public/Trace/Detail/LogScope.inl
+++ b/Engine/Source/Runtime/TraceLog/Public/Trace/Detail/LogScope.inl
@@ -26,7 +26,7 @@ extern TRACELOG_API uint32 volatile	GLogSerial;
 UE_AUTORTFM_ALWAYS_OPEN
 inline void FLogScope::Commit() const
 {
-	AtomicStoreRelease((uint8**) &(Buffer->Committed), Buffer->Cursor);
+	AtomicStoreRelease((uint8**) &(Buffer->Committed), Buffer->ReadCursor());
 }
 
 ////////////////////////////////////////////////////////////////////////////////
@@ -35,7 +35,7 @@ inline void FLogScope::Commit(FWriteBuffer* __restrict LatestBuffer) const
 {
 	if (LatestBuffer != Buffer)
 	{
-		AtomicStoreRelease((uint8**) &(LatestBuffer->Committed), LatestBuffer->Cursor);
+		AtomicStoreRelease((uint8**) &(LatestBuffer->Committed), LatestBuffer->ReadCursor());
 	}
 
 	Commit();
@@ -66,7 +66,7 @@ inline void FLogScope::EnterPrelude(uint32 Size)
 	uint32 AllocSize = sizeof(HeaderType) + Size;
 
 	Buffer = Writer_GetBuffer();
-	if (UNLIKELY(Buffer->Cursor + AllocSize > (uint8*)Buffer))
+	if (UNLIKELY(Buffer->ReadCursor() + AllocSize > (uint8*)Buffer))
 	{
 		Buffer = Writer_NextBuffer();
 	}
@@ -81,8 +81,8 @@ inline void FLogScope::EnterPrelude(uint32 Size)
 	}
 #endif
 
-	Ptr = Buffer->Cursor + sizeof(HeaderType);
-	Buffer->Cursor += AllocSize;
+	Ptr = Buffer->ReadCursor() + sizeof(HeaderType);
+	Buffer->MoveCursor((uint16)AllocSize);
 }
 
 ////////////////////////////////////////////////////////////////////////////////
@@ -120,8 +120,8 @@ inline void TLogScope<bMaybeHasAux>::operator += (const FLogScope&) const
 	if constexpr (bMaybeHasAux)
 	{
 		FWriteBuffer* LatestBuffer = Writer_GetBuffer();
-		LatestBuffer->Cursor[0] = uint8(EKnownEventUids::AuxDataTerminal << EKnownEventUids::_UidShift);
-		LatestBuffer->Cursor++;
+		*LatestBuffer->ReadCursor() = uint8(EKnownEventUids::AuxDataTerminal << EKnownEventUids::_UidShift);
+		LatestBuffer->MoveCursor(1u);
 
 		Commit(LatestBuffer);
 	}
@@ -152,15 +152,15 @@ inline void FScopedLogScope::Deinit()
 	uint8 LeaveUid = uint8(EKnownEventUids::LeaveScope << EKnownEventUids::_UidShift);
 
 	FWriteBuffer* Buffer = Writer_GetBuffer();
-	if (UNLIKELY(int32((uint8*)Buffer - Buffer->Cursor) < int32(sizeof(LeaveUid))))
+	if (UNLIKELY(int32((uint8*)Buffer - Buffer->ReadCursor()) < int32(sizeof(LeaveUid))))
 	{
 		Buffer = Writer_NextBuffer();
 	}
 
-	Buffer->Cursor[0] = LeaveUid;
-	Buffer->Cursor += sizeof(LeaveUid);
+	*Buffer->ReadCursor() = LeaveUid;
+	Buffer->MoveCursor(sizeof(LeaveUid));
 
-	AtomicStoreRelease((uint8**) &(Buffer->Committed), Buffer->Cursor);
+	AtomicStoreRelease((uint8**) &(Buffer->Committed), Buffer->ReadCursor());
 }
 
 ////////////////////////////////////////////////////////////////////////////////
@@ -212,17 +212,17 @@ inline void FScopedStampedLogScope::Deinit()
 	uint64 Stamp = TimeGetTimestamp() - GStartCycle;
 
 	FWriteBuffer* Buffer = Writer_GetBuffer();
-	if (UNLIKELY(int32((uint8*)Buffer - Buffer->Cursor) < int32(sizeof(Stamp))))
+	if (UNLIKELY(int32((uint8*)Buffer - Buffer->ReadCursor()) < int32(sizeof(Stamp))))
 	{
 		Buffer = Writer_NextBuffer();
 	}
 
 	Stamp <<= 8;
 	Stamp += uint8(EKnownEventUids::LeaveScope_TB) << EKnownEventUids::_UidShift;
-	memcpy((uint64*)(Buffer->Cursor), &Stamp, sizeof(Stamp));
-	Buffer->Cursor += sizeof(Stamp);
+	memcpy(Buffer->ReadCursor(), &Stamp, sizeof(Stamp));
+	Buffer->MoveCursor(sizeof(Stamp));
 
-	AtomicStoreRelease((uint8**) &(Buffer->Committed), Buffer->Cursor);
+	AtomicStoreRelease((uint8**) &(Buffer->Committed), Buffer->ReadCursor());
 }
 
 ////////////////////////////////////////////////////////////////////////////////
@@ -273,15 +273,15 @@ FORCENOINLINE auto FLogScope::ScopedEnter()
 	FWriteBuffer* Buffer = Writer_GetBuffer();
 	constexpr int32 EventHeaderSize = (EventType::EventFlags & FEventInfo::Flag_NoSync) != 0 ? sizeof(FEventHeader) : sizeof(FEventHeaderSync);
 	constexpr int32 RequiredSize = sizeof(EnterUid) + EventHeaderSize + EventType::GetSize();
-	if (UNLIKELY(int32((uint8*)Buffer - Buffer->Cursor) < RequiredSize))
+	if (UNLIKELY(int32((uint8*)Buffer - Buffer->ReadCursor()) < RequiredSize))
 	{
 		Buffer = Writer_NextBuffer();
 	}
 
-	Buffer->Cursor[0] = EnterUid;
-	Buffer->Cursor += sizeof(EnterUid);
+	*Buffer->ReadCursor() = EnterUid;
+	Buffer->MoveCursor(sizeof(EnterUid));
 
-	AtomicStoreRelease((uint8**) &(Buffer->Committed), Buffer->Cursor);
+	AtomicStoreRelease((uint8**) &(Buffer->Committed), Buffer->ReadCursor());
 
 	return Enter<EventType>();
 }
@@ -296,7 +296,7 @@ FORCENOINLINE auto FLogScope::ScopedStampedEnter()
 	FWriteBuffer* Buffer = Writer_GetBuffer();
 	constexpr int32 EventHeaderSize = (EventType::EventFlags & FEventInfo::Flag_NoSync) != 0 ? sizeof(FEventHeader) : sizeof(FEventHeaderSync);
 	constexpr int32 RequiredSize = sizeof(Stamp) + EventHeaderSize + EventType::GetSize();
-	if (UNLIKELY(int32((uint8*)Buffer - Buffer->Cursor) < RequiredSize))
+	if (UNLIKELY(int32((uint8*)Buffer - Buffer->ReadCursor()) < RequiredSize))
 	{
 		Buffer = Writer_NextBuffer();
 	}
@@ -304,10 +304,10 @@ FORCENOINLINE auto FLogScope::ScopedStampedEnter()
 	Stamp = TimeGetTimestamp() - GStartCycle;
 	Stamp <<= 8;
 	Stamp += uint8(EKnownEventUids::EnterScope_TB) << EKnownEventUids::_UidShift;
-	memcpy((uint64*)(Buffer->Cursor), &Stamp, sizeof(Stamp));
-	Buffer->Cursor += sizeof(Stamp);
+	memcpy(Buffer->ReadCursor(), &Stamp, sizeof(Stamp));
+	Buffer->MoveCursor(sizeof(Stamp));
 
-	AtomicStoreRelease((uint8**) &(Buffer->Committed), Buffer->Cursor);
+	AtomicStoreRelease((uint8**) &(Buffer->Committed), Buffer->ReadCursor());
 
 	return Enter<EventType>();
 }
diff --git a/Engine/Source/Runtime/TraceLog/Public/Trace/Detail/Writer.inl b/Engine/Source/Runtime/TraceLog/Public/Trace/Detail/Writer.inl
--- a/Engine/Source/Runtime/TraceLog/Public/Trace/Detail/Writer.inl
+++ b/Engine/Source/Runtime/TraceLog/Public/Trace/Detail/Writer.inl
@@ -9,20 +9,53 @@ namespace Trace {
 namespace Private {
 
 ////////////////////////////////////////////////////////////////////////////////
+/**
+ * Structure for managing event data buffers. Most fields are only accessed by the
+ * one thread; consuming thread (e.g. worker thread) or producing thread (e.g. any 
+ * thread tracing data). Volatile fields are used to safely communicate commited data
+ * and buffer usage.
+ */
 struct FWriteBuffer
 {
-	uint8						Overflow[8];
+	void						MoveCursor(uint16 Bytes);
+	void						SetCursor(uint8* Ptr);
+	uint8*						ReadCursor() const;
+
+	uint8						Overflow[6];
+	uint16						CursorMaskInv;
 	uint64						PrevTimestamp;
-	FWriteBuffer* __restrict	NextThread;
-	FWriteBuffer* volatile		NextBuffer;
-	uint8* __restrict			Cursor;
-	uint8* __restrict volatile	Committed;
-	uint8* __restrict			Reaped;
-	int32 volatile				EtxOffset;
 	uint16						Size;
 	uint16						ThreadId;
+	int32 volatile				EtxOffset; // Shared access (producer and consumer)
+	FWriteBuffer* __restrict	NextThread;
+	FWriteBuffer* volatile		NextBuffer; //Shared access
+	uint8* __restrict volatile	Committed; //Shared access
+	uint8* __restrict			Reaped;
+	uint8* __restrict			Cursor;
 };
 
+////////////////////////////////////////////////////////////////////////////////
+inline void FWriteBuffer::MoveCursor(uint16 Bytes)
+{
+	Cursor += (Bytes & ~CursorMaskInv);
+}
+
+////////////////////////////////////////////////////////////////////////////////
+inline void FWriteBuffer::SetCursor(uint8* Ptr)
+{
+	if (CursorMaskInv != 0)
+	{
+		return;
+	}
+	Cursor = Ptr;
+}
+
+////////////////////////////////////////////////////////////////////////////////
+inline uint8* FWriteBuffer::ReadCursor() const
+{
+	return Cursor;
+}
+
 ////////////////////////////////////////////////////////////////////////////////
 TRACELOG_API uint64				TimeGetTimestamp();
 TRACELOG_API FWriteBuffer*		Writer_NextBuffer();
diff --git a/Engine/Source/Runtime/TraceLog/Public/Trace/Trace.h b/Engine/Source/Runtime/TraceLog/Public/Trace/Trace.h
--- a/Engine/Source/Runtime/TraceLog/Public/Trace/Trace.h
+++ b/Engine/Source/Runtime/TraceLog/Public/Trace/Trace.h
@@ -116,6 +116,7 @@ enum class EMessageType : uint8
 	FileOpenError,
 	WriterError,
 	CompressionError,
+	PanicError,
 	// Fatal errors should cause application to stop
 	FatalStart			= 0x40,
 	GenericFatal,
@@ -128,7 +129,7 @@ struct FMessageEvent
 	EMessageType		Type;
 	/** Type of message stringified */
 	const char*			TypeStr;
-	/** Clarifying message, may be null for some message types. Pointer only valide during callback. */
+	/** Clarifying message, may be null for some message types. Pointer only valid during callback. */
 	const char*			Description;
 };
 
@@ -249,25 +250,25 @@ UE_TRACE_API void	SetMemoryHooks(AllocFunc Alloc, FreeFunc Free) UE_TRACE_IMPL()
 UE_TRACE_API void	SetMessageCallback(OnMessageFunc* MessageFunc) UE_TRACE_IMPL();
 
 /**
- * Set optional on update callbacks. If set issued after every update and once after initalization.
+ * Set optional on update callbacks. If set issued after every update and once after initialization.
  * @param UpdateFunc Function to call after every update
  */
 UE_TRACE_API void	SetUpdateCallback(OnUpdateFunc* UpdateFunc) UE_TRACE_IMPL();
 
 /**
- * Initalize TraceLog library.
- * @param Desc Initalization options
+ * Initialize TraceLog library.
+ * @param Desc Initialization options
  */
 UE_TRACE_API void	Initialize(const FInitializeDesc& Desc) UE_TRACE_IMPL();
 
 /**
- * Manually start worker thread if library is initalized without worker thread.
+ * Manually start worker thread if library is initialized without worker thread.
  */
 UE_TRACE_API void	StartWorkerThread() UE_TRACE_IMPL();
 
 /**
  * Call when application is exiting. Notifies TraceLog that the worker thread can 
- * dissapear at any time and clears block pool limits. Tracing is still possible.
+ * disappear at any time and clears block pool limits. Tracing is still possible.
  */
 UE_TRACE_API void	Exit() UE_TRACE_IMPL();
 
@@ -278,7 +279,8 @@ UE_TRACE_API void	Exit() UE_TRACE_IMPL();
 UE_TRACE_API void	Shutdown() UE_TRACE_IMPL();
 
 /**
- * Notifies TraceLog about a critical failure. Disables all tracing by muting all channels.
+ * Puts TraceLog in a critical failure stats. Disables all tracing by muting all channels.
+ * Tracing will be totally disabled for the remainder of the process lifetime.
  */
 UE_TRACE_API void	Panic() UE_TRACE_IMPL();
 
@@ -361,14 +363,14 @@ UE_TRACE_API bool	Stop() UE_TRACE_IMPL(false);
 
 /**
  * Checks if a string is valid channel name.
- * @param ChanelName String to check
+ * @param ChannelName String to check
  * @return True if there is a registered channel object with the same name.
  */
 UE_TRACE_API bool	IsChannel(const TCHAR* ChannelName) UE_TRACE_IMPL(false);
 
 /**
  * Toggles channel to control output of events.
- * @param ChanelName Name of channel to control
+ * @param ChannelName Name of channel to control
  * @param bEnabled Set if channel should be enabled or disabled
  * @return Final channel state. If the channel name was not found return false.
  */
@@ -389,7 +391,7 @@ UE_TRACE_API void	EnumerateChannels(ChannelIterFunc IterFunc, void* User) UE_TRA
 UE_TRACE_API void	EnumerateChannels(ChannelIterCallback IterFunc, void* User) UE_TRACE_IMPL();
 
 /**
- * Register a new thread in Trace. This is a requirement before tracing anything from the thrad.
+ * Register a new thread in Trace. This is a requirement before tracing anything from the thread.
  * @param Name Display name of thread
  * @param SystemId Platform specific thread id
  * @param SortHint Suggestion on how to order thread in presentation
@@ -465,9 +467,9 @@ UE_TRACE_API FChannel* FindChannel(FChannelId ChannelId) UE_TRACE_IMPL(nullptr);
 #define UE_TRACE_MINIMAL_EVENT_END()																	TRACE_IMPL_MINIMAL(EVENT_END)
 #define UE_TRACE_MINIMAL_LOG(LoggerName, EventName, ChannelsExpr, ...)									TRACE_IMPL_MINIMAL(LOG, LoggerName, EventName, ChannelsExpr, ##__VA_ARGS__)
 #define UE_TRACE_MINIMAL_LOG_SCOPED(LoggerName, EventName, ChannelsExpr, ...)							TRACE_IMPL_MINIMAL(LOG_SCOPED, LoggerName, EventName, ChannelsExpr, ##__VA_ARGS__)
-#define UE_TRACE_MINIMAL_LOG_SCOPED_CONDITIONAL(LoggerName, EventName, ChannelsExpr, Condition, ...)	TRACE_IMPL(LOG_SCOPED_CONDITIONAL, LoggerName, EventName, ChannelsExpr, Condition, ##__VA_ARGS__)
+#define UE_TRACE_MINIMAL_LOG_SCOPED_CONDITIONAL(LoggerName, EventName, ChannelsExpr, Condition, ...)	TRACE_IMPL_MINIMAL(LOG_SCOPED_CONDITIONAL, LoggerName, EventName, ChannelsExpr, Condition, ##__VA_ARGS__)
 #define UE_TRACE_MINIMAL_LOG_SCOPED_T(LoggerName, EventName, ChannelsExpr, ...)							TRACE_IMPL_MINIMAL(LOG_SCOPED_T, LoggerName, EventName, ChannelsExpr, ##__VA_ARGS__)
-#define UE_TRACE_MINIMAL_LOG_SCOPED_T_CONDITIONAL(LoggerName, EventName, ChannelsExpr, Condition, ...)	TRACE_IMPL(LOG_SCOPED_T_CONDITIONAL, LoggerName, EventName, ChannelsExpr, Condition, ##__VA_ARGS__)
+#define UE_TRACE_MINIMAL_LOG_SCOPED_T_CONDITIONAL(LoggerName, EventName, ChannelsExpr, Condition, ...)	TRACE_IMPL_MINIMAL(LOG_SCOPED_T_CONDITIONAL, LoggerName, EventName, ChannelsExpr, Condition, ##__VA_ARGS__)
 #define UE_TRACE_MINIMAL_GET_DEFINITION_TYPE_ID(LoggerName, EventName)									TRACE_IMPL_MINIMAL(GET_DEFINITION_TYPE_ID, LoggerName, EventName)
 #define UE_TRACE_MINIMAL_LOG_DEFINITION(LoggerName, EventName, Id, ChannelsExpr, ...)					TRACE_IMPL_MINIMAL(LOG_DEFINITION, LoggerName, EventName, Id, ChannelsExpr, ##__VA_ARGS__)
 #define UE_TRACE_MINIMAL_CHANNEL(ChannelName, ...)														TRACE_IMPL_MINIMAL(CHANNEL, ChannelName, ##__VA_ARGS__)
