Added asserts to Strncpy to ensure the output buffer has at least one byte.  This is because a caller expects a zero-terminated pointer result, and this cannot be the case if there is no output buffer.
Removed Strncpy usage from TraceLog because of a circular dependency problem.

diff --git a/Engine/Source/Runtime/Core/Private/GenericPlatform/GenericPlatformString.cpp b/Engine/Source/Runtime/Core/Private/GenericPlatform/GenericPlatformString.cpp
--- a/Engine/Source/Runtime/Core/Private/GenericPlatform/GenericPlatformString.cpp
+++ b/Engine/Source/Runtime/Core/Private/GenericPlatform/GenericPlatformString.cpp
@@ -25,6 +25,13 @@ void* FGenericPlatformString::Memcpy(void* Dest, const void* Src, SIZE_T Count)
 	return FMemory::Memcpy(Dest, Src, Count);
 }
 
+#if DO_CHECK
+void FGenericPlatformString::ReportZeroStrncpySize()
+{
+	checkf(false, TEXT("MaxLen cannot be zero as Strncpy expects to append a zero terminator"));
+}
+#endif
+
 namespace
 {
 	// Output iterator which will not overflow the destination buffer but will keep track of how many characters have been written.
diff --git a/Engine/Source/Runtime/Core/Private/GenericPlatform/GenericWidePlatformString.cpp b/Engine/Source/Runtime/Core/Private/GenericPlatform/GenericWidePlatformString.cpp
--- a/Engine/Source/Runtime/Core/Private/GenericPlatform/GenericWidePlatformString.cpp
+++ b/Engine/Source/Runtime/Core/Private/GenericPlatform/GenericWidePlatformString.cpp
@@ -32,6 +32,13 @@ WIDECHAR* FGenericWidePlatformString::Strcpy(WIDECHAR* Dest, const WIDECHAR* Src
 
 WIDECHAR* FGenericWidePlatformString::Strncpy(WIDECHAR* Dest, const WIDECHAR* Src, SIZE_T MaxLen)
 {
+#if DO_CHECK
+	if (UNLIKELY(MaxLen == 0))
+	{
+		ReportZeroStrncpySize();
+	}
+#endif
+
 	TCHAR *BufPtr = Dest;
 	
 	// the spec says that strncpy should fill the buffer with zeroes
diff --git a/Engine/Source/Runtime/Core/Public/GenericPlatform/GenericPlatformString.h b/Engine/Source/Runtime/Core/Public/GenericPlatform/GenericPlatformString.h
--- a/Engine/Source/Runtime/Core/Public/GenericPlatform/GenericPlatformString.h
+++ b/Engine/Source/Runtime/Core/Public/GenericPlatform/GenericPlatformString.h
@@ -357,4 +357,9 @@ private:
 	 * @return      Dest
 	 */
 	static CORE_API void* Memcpy(void* Dest, const void* Src, SIZE_T Count);
+
+protected:
+#if DO_CHECK
+	static CORE_API void ReportZeroStrncpySize();
+#endif
 };
diff --git a/Engine/Source/Runtime/Core/Public/GenericPlatform/GenericWidePlatformString.h b/Engine/Source/Runtime/Core/Public/GenericPlatform/GenericWidePlatformString.h
--- a/Engine/Source/Runtime/Core/Public/GenericPlatform/GenericWidePlatformString.h
+++ b/Engine/Source/Runtime/Core/Public/GenericPlatform/GenericWidePlatformString.h
@@ -228,6 +228,13 @@ PRAGMA_ENABLE_DEPRECATION_WARNINGS
 
 	static inline ANSICHAR* Strncpy(ANSICHAR* Dest, const ANSICHAR* Src, SIZE_T MaxLen)
 	{
+#if DO_CHECK
+		if (UNLIKELY(MaxLen == 0))
+		{
+			ReportZeroStrncpySize();
+		}
+#endif
+
 // Skip suggestions about using strncpy_s instead.
 PRAGMA_DISABLE_DEPRECATION_WARNINGS
 		::strncpy(Dest, Src, MaxLen);
diff --git a/Engine/Source/Runtime/Core/Public/GenericPlatform/StandardPlatformString.h b/Engine/Source/Runtime/Core/Public/GenericPlatform/StandardPlatformString.h
--- a/Engine/Source/Runtime/Core/Public/GenericPlatform/StandardPlatformString.h
+++ b/Engine/Source/Runtime/Core/Public/GenericPlatform/StandardPlatformString.h
@@ -61,6 +61,13 @@ PRAGMA_ENABLE_DEPRECATION_WARNINGS
 
 	static inline WIDECHAR* Strncpy(WIDECHAR* Dest, const WIDECHAR* Src, SIZE_T MaxLen)
 	{
+#if DO_CHECK
+		if (UNLIKELY(MaxLen == 0))
+		{
+			ReportZeroStrncpySize();
+		}
+#endif
+
 // Skip suggestions about using wcsncpy_s instead.
 PRAGMA_DISABLE_DEPRECATION_WARNINGS
 		wcsncpy(Dest, Src, MaxLen-1);
@@ -206,6 +213,13 @@ PRAGMA_ENABLE_DEPRECATION_WARNINGS
 
 	static inline ANSICHAR* Strncpy(ANSICHAR* Dest, const ANSICHAR* Src, SIZE_T MaxLen)
 	{
+#if DO_CHECK
+		if (UNLIKELY(MaxLen == 0))
+		{
+			ReportZeroStrncpySize();
+		}
+#endif
+
 // Skip suggestions about using strncpy_s instead.
 PRAGMA_DISABLE_DEPRECATION_WARNINGS
 		::strncpy(Dest, Src, MaxLen);
diff --git a/Engine/Source/Runtime/Core/Public/Microsoft/MicrosoftPlatformString.h b/Engine/Source/Runtime/Core/Public/Microsoft/MicrosoftPlatformString.h
--- a/Engine/Source/Runtime/Core/Public/Microsoft/MicrosoftPlatformString.h
+++ b/Engine/Source/Runtime/Core/Public/Microsoft/MicrosoftPlatformString.h
@@ -72,6 +72,13 @@ struct FMicrosoftPlatformString :
 
 	static inline WIDECHAR* Strncpy(WIDECHAR* Dest, const WIDECHAR* Src, SIZE_T MaxLen)
 	{
+#if DO_CHECK
+		if (UNLIKELY(MaxLen == 0))
+		{
+			ReportZeroStrncpySize();
+		}
+#endif
+
 		_tcsncpy(Dest, Src, MaxLen-1);
 		Dest[MaxLen-1] = 0;
 		return Dest;
@@ -202,6 +209,13 @@ struct FMicrosoftPlatformString :
 
 	static inline ANSICHAR* Strncpy(ANSICHAR* Dest, const ANSICHAR* Src, SIZE_T MaxLen)
 	{
+#if DO_CHECK
+		if (UNLIKELY(MaxLen == 0))
+		{
+			ReportZeroStrncpySize();
+		}
+#endif
+
 		strncpy(Dest, Src, MaxLen);
 		Dest[MaxLen-1] = 0;
 		return Dest;
diff --git a/Engine/Source/Runtime/TraceLog/Private/Trace/Control.cpp b/Engine/Source/Runtime/TraceLog/Private/Trace/Control.cpp
--- a/Engine/Source/Runtime/TraceLog/Private/Trace/Control.cpp
+++ b/Engine/Source/Runtime/TraceLog/Private/Trace/Control.cpp
@@ -305,7 +305,11 @@ void Writer_InitializeControl()
 					{
 						MaxLen = HostNameBufferSize;
 					}
-					FPlatformString::Strncpy(HostNameBuffer, ArgV[0], MaxLen);
+PRAGMA_DISABLE_DEPRECATION_WARNINGS
+					// We can't use FPlatformString::Strncpy because of circular module dependencies
+					strncpy(HostNameBuffer, ArgV[0], MaxLen);
+					HostNameBuffer[MaxLen - 1] = '\0';
+PRAGMA_ENABLE_DEPRECATION_WARNINGS
 					Host = HostNameBuffer;
 					Port = uint32(FPlatformString::Atoi(PortStr));
 				}
@@ -334,7 +338,11 @@ void Writer_InitializeControl()
 			ANSICHAR Channels[BufferSize] = {};
 			ANSICHAR* Ctx;
 			const bool bState = (ArgV[1][0] != '0');
-			FPlatformString::Strncpy(Channels, ArgV[0], BufferSize);
+PRAGMA_DISABLE_DEPRECATION_WARNINGS
+			// We can't use FPlatformString::Strncpy because of circular module dependencies
+			strncpy(Channels, ArgV[0], BufferSize);
+			Channels[BufferSize - 1] = '\0';
+PRAGMA_ENABLE_DEPRECATION_WARNINGS
 			ANSICHAR* Channel = FCStringAnsi::Strtok(Channels, ",", &Ctx);
 			while (Channel)
 			{
