TraceWriter: Fix incorrect units in timeout comparison (was comparing cycles to milliseconds).
#rnx

#Synced-CL Unknown

diff --git a/Engine/Source/Runtime/TraceLog/Private/Trace/BlockPool.cpp b/Engine/Source/Runtime/TraceLog/Private/Trace/BlockPool.cpp
--- a/Engine/Source/Runtime/TraceLog/Private/Trace/BlockPool.cpp
+++ b/Engine/Source/Runtime/TraceLog/Private/Trace/BlockPool.cpp
@@ -195,7 +195,8 @@ FWriteBuffer* Writer_AllocateBlockFromPool()
 		{
 			const uint64 Ts = TimeGetRelativeTimestamp();
 			StallStart = StallStart ? StallStart : Ts;
-			if ((Ts - StallStart) > GPoolLimitTimeoutMs)
+			uint32 DurationMs = uint32(1000.f*(float(Ts - StallStart) / float(TimeGetFrequency())));
+			if (DurationMs > GPoolLimitTimeoutMs)
 			{
 				// At this point we have timed out waiting for a free block.
 				// This likely means the worker thread cannot keep up with
