[TraceLog]
* Added TRACELOG_API to UE::Trace::Private::Encode() and UE::Trace::Private::GetEncodeMaxSize() to be accessible from other friendly modules.
* Extracted Protocol6::FNewEventField struct out of Protocol6::FNewEventEvent.
* Added None=0 enum field to Protocol6::EEventFlags.

#rnx
#Synced-CL 49171203

diff --git a/Engine/Source/Runtime/TraceLog/Private/Trace/Codec.cpp b/Engine/Source/Runtime/TraceLog/Private/Trace/Codec.cpp
--- a/Engine/Source/Runtime/TraceLog/Private/Trace/Codec.cpp
+++ b/Engine/Source/Runtime/TraceLog/Private/Trace/Codec.cpp
@@ -29,50 +29,32 @@ namespace Trace {
 namespace Private {
 
 ////////////////////////////////////////////////////////////////////////////////
-int32 Encode(const void* Src, int32 SrcSize, void* Dest, int32 DestSize)
+TRACELOG_API uint32 GetEncodeMaxSize(uint32 InputSize)
 {
-#if TRACE_PRIVATE_MINIMAL_ENABLED
-	FProfilerScope _(__func__);
-#endif
-	int Result = TRACE_PRIVATE_LZ4_NAMESPACE LZ4_compress_fast(
-		(const char*)Src,
-		(char*)Dest,
-		SrcSize,
-		DestSize,
-		1 // increase by 1 for small speed increase
-	);
-	if (!Result)
-	{
-#if UE_TRACE_ENABLED 
-		UE_TRACE_MESSAGE_F(CompressionError, "LZ4 failed to compress %d bytes (dest %d bytes).", SrcSize, DestSize);
-#endif
-	}
-	return Result;
+	return LZ4_COMPRESSBOUND(InputSize);
 }
 
 ////////////////////////////////////////////////////////////////////////////////
-int32 EncodeNoInstr(const void* Src, int32 SrcSize, void* Dest, int32 DestSize)
+TRACELOG_API int32 Encode(const void* Src, int32 SrcSize, void* Dest, int32 DestSize)
 {
-	int Result = TRACE_PRIVATE_LZ4_NAMESPACE LZ4_compress_fast(
-		(const char*)Src,
-		(char*)Dest,
-		SrcSize,
-		DestSize,
-		1 // increase by 1 for small speed increase
-	);
-	if (!Result)
-	{
-		// Cannot emit error message here since it might cause
-		// events.
-		//PLATFORM_BREAK();
-	}
-	return Result;
+	constexpr int Acceleration = 1; // increase by 1 for small speed increase
+	return TRACE_PRIVATE_LZ4_NAMESPACE LZ4_compress_fast((const char*)Src, (char*)Dest, SrcSize, DestSize, Acceleration);
 }
 
 ////////////////////////////////////////////////////////////////////////////////
-uint32 GetEncodeMaxSize(uint32 InputSize)
+int32 EncodeInstr(const void* Src, int32 SrcSize, void* Dest, int32 DestSize)
 {
-	return LZ4_COMPRESSBOUND(InputSize);
+#if TRACE_PRIVATE_MINIMAL_ENABLED
+	FProfilerScope _(__func__);
+#endif
+	int Result = Encode(Src, SrcSize, Dest, DestSize);
+#if UE_TRACE_ENABLED
+	if (!Result)
+	{
+		UE_TRACE_MESSAGE_F(CompressionError, "LZ4 failed to compress %d bytes (dest %d bytes).", SrcSize, DestSize);
+	}
+#endif
+	return Result;
 }
 
 ////////////////////////////////////////////////////////////////////////////////
diff --git a/Engine/Source/Runtime/TraceLog/Private/Trace/EventNode.cpp b/Engine/Source/Runtime/TraceLog/Private/Trace/EventNode.cpp
--- a/Engine/Source/Runtime/TraceLog/Private/Trace/EventNode.cpp
+++ b/Engine/Source/Runtime/TraceLog/Private/Trace/EventNode.cpp
@@ -135,9 +135,9 @@ void FEventNode::Describe() const
 		NamesSize += Info->Fields[i].NameSize;
 	}
 
-	// Allocate the new event event in the log stream.
+	// Allocate the "NewEvent" event in the log stream.
 	uint32 EventSize = sizeof(FNewEventEvent);
-	EventSize += sizeof(FNewEventEvent::Fields[0]) * Info->FieldCount;
+	EventSize += sizeof(FNewEventField) * Info->FieldCount;
 	EventSize += NamesSize;
 	EventSize = (EventSize + 1) & ~1; // align to 2 to keep UBSAN happy
 
@@ -164,7 +164,7 @@ void FEventNode::Describe() const
 	for (uint32 i = 0; i < Info->FieldCount; ++i)
 	{
 		const FFieldDesc& Field = Info->Fields[i];
-		auto& Out = Event.Fields[i];
+		FNewEventField& Out = Event.Fields[i];
 		if (i == DefinitionIdFieldIdx)
 		{
 			Out.FieldType = EFieldFamily::DefinitionId;
diff --git a/Engine/Source/Runtime/TraceLog/Private/Trace/Important/Cache.cpp b/Engine/Source/Runtime/TraceLog/Private/Trace/Important/Cache.cpp
--- a/Engine/Source/Runtime/TraceLog/Private/Trace/Important/Cache.cpp
+++ b/Engine/Source/Runtime/TraceLog/Private/Trace/Important/Cache.cpp
@@ -18,7 +18,7 @@ namespace Private {
 
 ////////////////////////////////////////////////////////////////////////////////
 uint32	GetEncodeMaxSize(uint32);
-int32	Encode(const void*, int32, void*, int32);
+int32	EncodeInstr(const void*, int32, void*, int32);
 void*	Writer_MemoryAllocate(SIZE_T, uint32);
 void	Writer_MemoryFree(void*, uint32);
 void	Writer_SendDataRaw(const void*, uint32);
@@ -68,7 +68,7 @@ static FCacheBuffer* Writer_CacheCreateBuffer(uint32 Size)
 static void Writer_CacheCommit(const FCacheBuffer* Collector)
 {
 	FProfilerScope _(__func__);
-	
+
 	// Check there's enough size to compress Data into.
 	uint32 InputSize = uint32(Collector->Size - Collector->Remaining);
 	uint32 EncodeMaxSize = GetEncodeMaxSize(InputSize);
@@ -91,7 +91,7 @@ static void Writer_CacheCommit(const FCacheBuffer* Collector)
 	const uint32 ActiveBufferOffset = GCacheActiveBuffer->Size - GCacheActiveBuffer->Remaining;
 	uint32 PacketTotalSize = 0;
 	auto* Packet = (FTidPacketEncoded*)(GCacheActiveBuffer->Data + ActiveBufferOffset);
-	const int32 CompressedSize = Encode(Collector->Data, InputSize, Packet->Data, EncodeMaxSize);
+	const int32 CompressedSize = EncodeInstr(Collector->Data, InputSize, Packet->Data, EncodeMaxSize);
 	if (CompressedSize > 0)
 	{
 		Packet->PacketSize = uint16(CompressedSize + sizeof(FTidPacketEncoded));
@@ -127,7 +127,7 @@ static void Writer_CacheCommit(const FCacheBuffer* Collector)
 void Writer_CacheData(uint8* Data, uint32 Size)
 {
 	FProfilerScope _(__func__);
-	
+
 	Writer_SendData(ETransportTid::Importants, Data, Size);
 
 	if (GCacheCollector == nullptr)
@@ -163,7 +163,7 @@ void Writer_CacheData(uint8* Data, uint32 Size)
 void Writer_CacheOnConnect()
 {
 	FProfilerScope _(__func__);
-	
+
 	if (GCacheCollector == nullptr)
 	{
 		return;
diff --git a/Engine/Source/Runtime/TraceLog/Private/Trace/Tail.cpp b/Engine/Source/Runtime/TraceLog/Private/Trace/Tail.cpp
--- a/Engine/Source/Runtime/TraceLog/Private/Trace/Tail.cpp
+++ b/Engine/Source/Runtime/TraceLog/Private/Trace/Tail.cpp
@@ -22,7 +22,7 @@ static_assert(ETransport::Active == ETransport::TidPacketSync, "Tail-tracing is
 
 ////////////////////////////////////////////////////////////////////////////////
 uint32		GetEncodeMaxSize(uint32);
-int32		Encode(const void*, int32, void*, int32);
+int32		EncodeInstr(const void*, int32, void*, int32);
 void*		Writer_MemoryAllocate(SIZE_T, uint32);
 void		Writer_MemoryFree(void*, uint32);
 void		Writer_SendData(uint32, uint8* __restrict, uint32);
@@ -255,7 +255,7 @@ void Writer_TailAppend(uint32 ThreadId, uint8* __restrict Data, uint32 Size)
 #endif
 		Packet->DecodedSize = uint16(Size);
 
-		int32 EncodedSize = Encode(Data, Size, Packet->Data, EncodeMaxSize);
+		int32 EncodedSize = EncodeInstr(Data, Size, Packet->Data, EncodeMaxSize);
 		if (EncodedSize > 0)
 		{
 			// Reclaim buffer space (worst case size - actual size)
diff --git a/Engine/Source/Runtime/TraceLog/Private/Trace/Writer.cpp b/Engine/Source/Runtime/TraceLog/Private/Trace/Writer.cpp
--- a/Engine/Source/Runtime/TraceLog/Private/Trace/Writer.cpp
+++ b/Engine/Source/Runtime/TraceLog/Private/Trace/Writer.cpp
@@ -37,8 +37,8 @@ namespace Trace {
 namespace Private {
 
 ////////////////////////////////////////////////////////////////////////////////
+int32			EncodeInstr(const void*, int32, void*, int32);
 int32			Encode(const void*, int32, void*, int32);
-int32			EncodeNoInstr(const void*, int32, void*, int32);
 void			Writer_SetUpdateCallback(OnUpdateFunc* Callback);
 void			Writer_SendData(uint32, uint8* __restrict, uint32);
 void			Writer_InitializeTail(int32);
@@ -452,7 +452,7 @@ void Writer_SendDataNoInstr(uint32 ThreadId, uint8* __restrict Data, uint32 Size
 		TTidPacketEncoded<MaxEncodedBuffer + EncodingOverhead> Packet;
 
 		Packet.DecodedSize = uint16(Size);
-		Packet.PacketSize = uint16(EncodeNoInstr(Data, Packet.DecodedSize, Packet.Data, sizeof(Packet.Data)));
+		Packet.PacketSize = uint16(Encode(Data, Packet.DecodedSize, Packet.Data, sizeof(Packet.Data)));
 
 		if (Packet.PacketSize > 0 &&
 			Packet.PacketSize < (Packet.DecodedSize + EncodingOverhead))
@@ -523,7 +523,7 @@ void Writer_SendData(uint32 ThreadId, uint8* __restrict Data, uint32 Size)
 		TTidPacketEncoded<MaxEncodedBuffer + EncodingOverhead> Packet;
 
 		Packet.DecodedSize = uint16(Size);
-		Packet.PacketSize = uint16(Encode(Data, Packet.DecodedSize, Packet.Data, sizeof(Packet.Data)));
+		Packet.PacketSize = uint16(EncodeInstr(Data, Packet.DecodedSize, Packet.Data, sizeof(Packet.Data)));
 
 		if (Packet.PacketSize > 0 &&
 			Packet.PacketSize < (Packet.DecodedSize + EncodingOverhead))
@@ -661,7 +661,7 @@ static void Writer_TraceStats(const FStatistics& InStats)
 	if (GDataHandle)
 	{
 		const uint64 CurrentTimeCycle = TimeGetRelativeTimestamp();
-		const float BlockPoolFreeFraction = InStats.BlockPoolAllocatedBlocks ? 
+		const float BlockPoolFreeFraction = InStats.BlockPoolAllocatedBlocks ?
 			float(InStats.BlockPoolFreeBlocks) / float(InStats.BlockPoolAllocatedBlocks) : 1.0f;
 		const float BlockPoolUsage = 1.0f - BlockPoolFreeFraction;
 		const uint32 DeltaEmitted = uint32(InStats.BytesEmitted - LastEmitted);
@@ -1021,10 +1021,10 @@ static void Writer_InternalInitializeImpl()
 	}
 
 	double StartDateTime = std::chrono::duration_cast<std::chrono::duration<double>>
-	(std::chrono::system_clock::now().time_since_epoch()).count();
-	
+		(std::chrono::system_clock::now().time_since_epoch()).count();
+
 	GStartCycle = TimeGetTimestamp();
-	
+
 	Writer_InitializeSharedBuffers();
 	Writer_InitializePool();
 	Writer_InitializeControl();
diff --git a/Engine/Source/Runtime/TraceLog/Public/Trace/Detail/Protocols/Protocol6.h b/Engine/Source/Runtime/TraceLog/Public/Trace/Detail/Protocols/Protocol6.h
--- a/Engine/Source/Runtime/TraceLog/Public/Trace/Detail/Protocols/Protocol6.h
+++ b/Engine/Source/Runtime/TraceLog/Public/Trace/Detail/Protocols/Protocol6.h
@@ -27,6 +27,7 @@ using Protocol5::EKnownEventUids;
 ////////////////////////////////////////////////////////////////////////////////
 enum class EEventFlags : uint8
 {
+	None			= 0,
 	Important		= 1 << 0,
 	MaybeHasAux		= 1 << 1,
 	NoSync			= 1 << 2,
@@ -42,42 +43,47 @@ enum EFieldFamily : uint8
 };
 	
 ////////////////////////////////////////////////////////////////////////////////
-struct FNewEventEvent
+struct FNewEventField
 {
-	uint16		EventUid;
-	uint8		FieldCount;
-	uint8		Flags;
-	uint8		LoggerNameSize;
-	uint8		EventNameSize;
-	struct
+	uint8	FieldType; // EFieldFamily
+	uint8	Unused;
+	union
 	{
-		uint8	FieldType;
-		union
+		struct
+		{
+			uint16	Offset;
+			uint16	Size;
+			uint8	TypeInfo; // EFieldType
+			uint8	NameSize;
+		} Regular;
+		struct
 		{
-			struct
-			{
-				uint16	Offset;
-				uint16	Size;
-				uint8	TypeInfo;
-				uint8	NameSize;
-			} Regular;
-			struct
-			{
-				uint16	Offset;
-				uint16	RefUid; 
-				uint8	TypeInfo;
-				uint8	NameSize;
-			} Reference;
-			struct
-			{
-				uint16	Offset;
-				uint16	Unused1;
-				uint8	Unused2;
-				uint8	TypeInfo;
-			} DefinitionId;
-		};
-	} Fields[];
+			uint16	Offset;
+			uint16	RefUid;
+			uint8	TypeInfo; // EFieldType
+			uint8	NameSize;
+		} Reference;
+		struct
+		{
+			uint16	Offset;
+			uint16	Unused1;
+			uint8	Unused2;
+			uint8	TypeInfo; // EFieldType
+		} DefinitionId;
+	};
+};
+static_assert(sizeof(FNewEventField) == 8, "Struct layout assumption doesn't match expectation");
+
+struct FNewEventEvent
+{
+	uint16			EventUid;
+	uint8			FieldCount;
+	uint8			Flags; // EEventFlags
+	uint8			LoggerNameSize;
+	uint8			EventNameSize;
+	FNewEventField	Fields[];
 };
+static_assert(sizeof(FNewEventEvent) == 6, "Struct layout assumption doesn't match expectation");
 	
 ////////////////////////////////////////////////////////////////////////////////
 
diff --git a/Engine/Source/Runtime/TraceLog/Public/Trace/Detail/Protocols/Protocol7.h b/Engine/Source/Runtime/TraceLog/Public/Trace/Detail/Protocols/Protocol7.h
--- a/Engine/Source/Runtime/TraceLog/Public/Trace/Detail/Protocols/Protocol7.h
+++ b/Engine/Source/Runtime/TraceLog/Public/Trace/Detail/Protocols/Protocol7.h
@@ -24,6 +24,7 @@ using Protocol6::FEventHeaderSync;
 using Protocol6::FAuxHeader;
 using Protocol6::EEventFlags;
 using Protocol6::EFieldFamily;
+using Protocol6::FNewEventField;
 using Protocol6::FNewEventEvent;
 
 ////////////////////////////////////////////////////////////////////////////////
