[Insights] Trace platform support

- Expose method to get native socket.
- Add function to get peer name for a socket.
- Expose method to control blocking status of a socket.

#Synced-CL 49276880

diff --git a/Engine/Source/Runtime/TraceLog/Private/Trace/Detail/Android/AndroidTrace.cpp b/Engine/Source/Runtime/TraceLog/Private/Trace/Detail/Android/AndroidTrace.cpp
--- a/Engine/Source/Runtime/TraceLog/Private/Trace/Detail/Android/AndroidTrace.cpp
+++ b/Engine/Source/Runtime/TraceLog/Private/Trace/Detail/Android/AndroidTrace.cpp
@@ -72,9 +72,14 @@ uint64 TimeGetTimestamp()
 }
 
 
+////////////////////////////////////////////////////////////////////////////////
+int TcpSocketNative(UPTRINT Socket)
+{
+	return int(Socket-1);
+}
 
 ////////////////////////////////////////////////////////////////////////////////
-static bool TcpSocketSetNonBlocking(int Socket, bool bNonBlocking)
+static bool TcpSocketSetNonBlockingInternal(int Socket, bool bNonBlocking)
 {
 	int Flags = fcntl(Socket, F_GETFL, 0);
 	if (Flags == -1)
@@ -86,6 +91,13 @@ static bool TcpSocketSetNonBlocking(int Socket, bool bNonBlocking)
 	return fcntl(Socket, F_SETFL, Flags) >= 0;
 }
 
+////////////////////////////////////////////////////////////////////////////////
+bool TcpSocketSetNonBlocking(UPTRINT Socket, bool bNonBlocking)
+{
+	int Inner = TcpSocketNative(Socket);
+	return TcpSocketSetNonBlockingInternal(Inner, bNonBlocking);
+}
+
 ////////////////////////////////////////////////////////////////////////////////
 UPTRINT TcpSocketConnect(const ANSICHAR* Host, uint16 Port)
 {
@@ -107,7 +119,7 @@ UPTRINT TcpSocketConnect(const ANSICHAR* Host, uint16 Port)
 		return 0;
 	}
 
-	if (!TcpSocketSetNonBlocking(Socket, false))
+	if (!TcpSocketSetNonBlockingInternal(Socket, false))
 	{
 		close(Socket);
 		return 0;
@@ -143,7 +155,7 @@ UPTRINT TcpSocketListen(uint16 Port)
 		return 0;
 	}
 
-	if (!TcpSocketSetNonBlocking(Socket, true))
+	if (!TcpSocketSetNonBlockingInternal(Socket, true))
 	{
 		close(Socket);
 		return 0;
@@ -153,9 +165,9 @@ UPTRINT TcpSocketListen(uint16 Port)
 }
 
 ////////////////////////////////////////////////////////////////////////////////
-int32 TcpSocketAccept(UPTRINT Socket, UPTRINT& Out)
+int32 TcpSocketAccept(UPTRINT Socket, UPTRINT& Out, bool bAsNonBlocking /*= false*/)
 {
-	int Inner = Socket - 1;
+	int Inner = TcpSocketNative(Socket);
 
 	Inner = accept(Inner, nullptr, nullptr);
 	if (Inner < 0)
@@ -163,7 +175,7 @@ int32 TcpSocketAccept(UPTRINT Socket, UPTRINT& Out)
 		return (errno == EAGAIN || errno == EWOULDBLOCK) - 1; // 0 if would block else -1
 	}
 
-	if (!TcpSocketSetNonBlocking(Inner, false))
+	if (!TcpSocketSetNonBlockingInternal(Inner, bAsNonBlocking))
 	{
 		close(Inner);
 		return 0;
@@ -176,7 +188,7 @@ int32 TcpSocketAccept(UPTRINT Socket, UPTRINT& Out)
 ////////////////////////////////////////////////////////////////////////////////
 bool TcpSocketHasData(UPTRINT Socket)
 {
-	int Inner = Socket - 1;
+	int Inner = TcpSocketNative(Socket);
 	fd_set FdSet;
 	FD_ZERO(&FdSet);
 	FD_SET(Inner, &FdSet);
@@ -185,6 +197,24 @@ bool TcpSocketHasData(UPTRINT Socket)
 }
 
 
+////////////////////////////////////////////////////////////////////////////////
+bool TcpSocketPeerName(UPTRINT Socket, char* Name, uint32 NameSize)
+{
+	int Inner = TcpSocketNative(Socket);
+	sockaddr_storage Addr;
+	socklen_t AddrLen = sizeof(Addr);
+
+	getpeername(Inner, (sockaddr*) &Addr, &AddrLen);
+	if (Addr.ss_family != AF_INET)
+	{
+		return false;
+	}
+
+	sockaddr_in* Ip4Addr = (sockaddr_in*) &Addr;
+	inet_ntop(AF_INET, &Ip4Addr->sin_addr, Name, NameSize);
+
+	return true;
+}
 
 ////////////////////////////////////////////////////////////////////////////////
 bool IoWrite(UPTRINT Handle, const void* Data, uint32 Size)
@@ -207,8 +237,6 @@ void IoClose(UPTRINT Handle)
 	close(Inner);
 }
 
-
-
 ////////////////////////////////////////////////////////////////////////////////
 UPTRINT FileOpen(const ANSICHAR* Path)
 {
diff --git a/Engine/Source/Runtime/TraceLog/Private/Trace/Detail/Apple/AppleTrace.cpp b/Engine/Source/Runtime/TraceLog/Private/Trace/Detail/Apple/AppleTrace.cpp
--- a/Engine/Source/Runtime/TraceLog/Private/Trace/Detail/Apple/AppleTrace.cpp
+++ b/Engine/Source/Runtime/TraceLog/Private/Trace/Detail/Apple/AppleTrace.cpp
@@ -1,6 +1,8 @@
 // Copyright Epic Games, Inc. All Rights Reserved.
 
 #include "Trace/Config.h"
+#include "Trace/Platform.h"
+#include <netinet/in.h>
 
 #if TRACE_PRIVATE_MINIMAL_ENABLED && PLATFORM_APPLE
 
@@ -74,9 +76,14 @@ TRACELOG_API uint64 TimeGetTimestamp()
 }
 
 
+////////////////////////////////////////////////////////////////////////////////
+int TcpSocketNative(UPTRINT Socket)
+{
+	return int(Socket-1);
+}
 
 ////////////////////////////////////////////////////////////////////////////////
-static bool TcpSocketSetNonBlocking(int Socket, bool bNonBlocking)
+static bool TcpSocketSetNonBlockingInternal(int Socket, bool bNonBlocking)
 {
 	int Flags = fcntl(Socket, F_GETFL, 0);
 	if (Flags == -1)
@@ -88,6 +95,13 @@ static bool TcpSocketSetNonBlocking(int Socket, bool bNonBlocking)
 	return fcntl(Socket, F_SETFL, Flags) >= 0;
 }
 
+////////////////////////////////////////////////////////////////////////////////
+bool TcpSocketSetNonBlocking(UPTRINT Socket, bool bNonBlocking)
+{
+	int Inner = TcpSocketNative(Socket);
+	return TcpSocketSetNonBlockingInternal(Inner, bNonBlocking);
+}
+
 ////////////////////////////////////////////////////////////////////////////////
 UPTRINT TcpSocketConnect(const ANSICHAR* Host, uint16 Port)
 {
@@ -141,7 +155,7 @@ UPTRINT TcpSocketConnect(const ANSICHAR* Host, uint16 Port)
 		return 0;
 	}
 
-	if (!TcpSocketSetNonBlocking(Socket, false))
+	if (!TcpSocketSetNonBlockingInternal(Socket, false))
 	{
 		close(Socket);
 		return 0;
@@ -177,7 +191,7 @@ UPTRINT TcpSocketListen(uint16 Port)
 		return 0;
 	}
 
-	if (!TcpSocketSetNonBlocking(Socket, true))
+	if (!TcpSocketSetNonBlockingInternal(Socket, true))
 	{
 		close(Socket);
 		return 0;
@@ -187,9 +201,9 @@ UPTRINT TcpSocketListen(uint16 Port)
 }
 
 ////////////////////////////////////////////////////////////////////////////////
-int32 TcpSocketAccept(UPTRINT Socket, UPTRINT& Out)
+int32 TcpSocketAccept(UPTRINT Socket, UPTRINT& Out, bool bAsNonBlocking /*= false*/)
 {
-	int Inner = int(Socket - 1);
+	int Inner = TcpSocketNative(Socket);
 
 	Inner = accept(Inner, nullptr, nullptr);
 	if (Inner < 0)
@@ -197,7 +211,7 @@ int32 TcpSocketAccept(UPTRINT Socket, UPTRINT& Out)
 		return (errno == EAGAIN || errno == EWOULDBLOCK) - 1; // 0 if would block else -1
 	}
 
-	if (!TcpSocketSetNonBlocking(Inner, false))
+	if (!TcpSocketSetNonBlockingInternal(Inner, bAsNonBlocking))
 	{
 		close(Inner);
 		return 0;
@@ -210,7 +224,7 @@ int32 TcpSocketAccept(UPTRINT Socket, UPTRINT& Out)
 ////////////////////////////////////////////////////////////////////////////////
 bool TcpSocketHasData(UPTRINT Socket)
 {
-	int Inner = int(Socket - 1);
+	int Inner = TcpSocketNative(Socket);
 	fd_set FdSet;
 	FD_ZERO(&FdSet);
 	FD_SET(Inner, &FdSet);
@@ -219,7 +233,24 @@ bool TcpSocketHasData(UPTRINT Socket)
 	return ((result != 0) || ((result == -1) && (errno == ETIMEDOUT)));
 }
 
+////////////////////////////////////////////////////////////////////////////////
+bool TcpSocketPeerName(UPTRINT Socket, char* Name, uint32 NameSize)
+{
+	int Inner = TcpSocketNative(Socket);
+	sockaddr_storage Addr;
+	socklen_t AddrLen = sizeof(Addr);
 
+	getpeername(Inner, (sockaddr*) &Addr, &AddrLen);
+	if (Addr.ss_family != AF_INET)
+	{
+		return false;
+	}
+
+	sockaddr_in* Ip4Addr = (sockaddr_in*) &Addr;
+	inet_ntop(AF_INET, &Ip4Addr->sin_addr, Name, NameSize);
+
+	return true;
+}
 
 ////////////////////////////////////////////////////////////////////////////////
 bool IoWrite(UPTRINT Handle, const void* Data, uint32 Size)
diff --git a/Engine/Source/Runtime/TraceLog/Private/Trace/Detail/Unix/UnixTrace.cpp b/Engine/Source/Runtime/TraceLog/Private/Trace/Detail/Unix/UnixTrace.cpp
--- a/Engine/Source/Runtime/TraceLog/Private/Trace/Detail/Unix/UnixTrace.cpp
+++ b/Engine/Source/Runtime/TraceLog/Private/Trace/Detail/Unix/UnixTrace.cpp
@@ -4,6 +4,7 @@
 
 #if TRACE_PRIVATE_MINIMAL_ENABLED && PLATFORM_UNIX
 
+#include <arpa/inet.h>
 #include <cstring>
 #include <errno.h>
 #include <fcntl.h>
@@ -74,10 +75,14 @@ TRACELOG_API uint64 TimeGetTimestamp()
 	return static_cast<uint64>(static_cast<uint64>(TimeSpec.tv_sec) * 10000000ULL + static_cast<uint64>(TimeSpec.tv_nsec) / 100ULL);
 }
 
-
+////////////////////////////////////////////////////////////////////////////////
+int TcpSocketNative(UPTRINT Socket)
+{
+	return int(Socket-1);
+}
 
 ////////////////////////////////////////////////////////////////////////////////
-static bool TcpSocketSetNonBlocking(int Socket, bool bNonBlocking)
+static bool TcpSocketSetNonBlockingInternal(int Socket, bool bNonBlocking)
 {
 	int Flags = fcntl(Socket, F_GETFL, 0);
 	if (Flags == -1)
@@ -89,6 +94,13 @@ static bool TcpSocketSetNonBlocking(int Socket, bool bNonBlocking)
 	return fcntl(Socket, F_SETFL, Flags) >= 0;
 }
 
+////////////////////////////////////////////////////////////////////////////////
+bool TcpSocketSetNonBlocking(UPTRINT Socket, bool bNonBlocking)
+{
+	int Inner = TcpSocketNative(Socket);
+	return TcpSocketSetNonBlockingInternal(Inner, bNonBlocking);
+}
+
 ////////////////////////////////////////////////////////////////////////////////
 UPTRINT TcpSocketConnect(const ANSICHAR* Host, uint16 Port)
 {
@@ -138,7 +150,7 @@ UPTRINT TcpSocketConnect(const ANSICHAR* Host, uint16 Port)
 		return 0;
 	}
 
-	if (!TcpSocketSetNonBlocking(Socket, false))
+	if (!TcpSocketSetNonBlockingInternal(Socket, false))
 	{
 		close(Socket);
 		return 0;
@@ -174,7 +186,7 @@ UPTRINT TcpSocketListen(uint16 Port)
 		return 0;
 	}
 
-	if (!TcpSocketSetNonBlocking(Socket, true))
+	if (!TcpSocketSetNonBlockingInternal(Socket, true))
 	{
 		close(Socket);
 		return 0;
@@ -184,9 +196,9 @@ UPTRINT TcpSocketListen(uint16 Port)
 }
 
 ////////////////////////////////////////////////////////////////////////////////
-int32 TcpSocketAccept(UPTRINT Socket, UPTRINT& Out)
+int32 TcpSocketAccept(UPTRINT Socket, UPTRINT& Out, bool bAsNonBlocking /* = false */)
 {
-	int Inner = Socket - 1;
+	int Inner = TcpSocketNative(Socket);
 
 	Inner = accept(Inner, nullptr, nullptr);
 	if (Inner < 0)
@@ -194,7 +206,7 @@ int32 TcpSocketAccept(UPTRINT Socket, UPTRINT& Out)
 		return (errno == EAGAIN || errno == EWOULDBLOCK) - 1; // 0 if would block else -1
 	}
 
-	if (!TcpSocketSetNonBlocking(Inner, false))
+	if (!TcpSocketSetNonBlockingInternal(Inner, bAsNonBlocking))
 	{
 		close(Inner);
 		return 0;
@@ -207,7 +219,7 @@ int32 TcpSocketAccept(UPTRINT Socket, UPTRINT& Out)
 ////////////////////////////////////////////////////////////////////////////////
 bool TcpSocketHasData(UPTRINT Socket)
 {
-	int Inner = Socket - 1;
+	int Inner = TcpSocketNative(Socket);
 	fd_set FdSet;
 	FD_ZERO(&FdSet);
 	FD_SET(Inner, &FdSet);
@@ -215,7 +227,23 @@ bool TcpSocketHasData(UPTRINT Socket)
 	return (select(Inner + 1, &FdSet, nullptr, nullptr, &TimeVal) != 0);
 }
 
+////////////////////////////////////////////////////////////////////////////////
+bool TcpSocketPeerName(UPTRINT Socket, char* Name, uint32 NameSize)
+{
+	sockaddr_storage Addr;
+	socklen_t AddrLen = sizeof(Addr);
+
+	getpeername(int(Socket-1), (sockaddr*) &Addr, &AddrLen);
+	if (Addr.ss_family != AF_INET)
+	{
+		return false;
+	}
 
+	sockaddr_in* Ip4Addr = (sockaddr_in*) &Addr;
+	inet_ntop(AF_INET, &Ip4Addr->sin_addr, Name, NameSize);
+
+	return true;
+}
 
 ////////////////////////////////////////////////////////////////////////////////
 bool IoWrite(UPTRINT Handle, const void* Data, uint32 Size)
diff --git a/Engine/Source/Runtime/TraceLog/Private/Trace/Detail/Windows/WindowsTrace.cpp b/Engine/Source/Runtime/TraceLog/Private/Trace/Detail/Windows/WindowsTrace.cpp
--- a/Engine/Source/Runtime/TraceLog/Private/Trace/Detail/Windows/WindowsTrace.cpp
+++ b/Engine/Source/Runtime/TraceLog/Private/Trace/Detail/Windows/WindowsTrace.cpp
@@ -91,12 +91,25 @@ static void TcpSocketInitialize()
 }
 
 ////////////////////////////////////////////////////////////////////////////////
-static bool TcpSocketSetNonBlocking(SOCKET Socket, bool bNonBlocking)
+int TcpSocketNative(UPTRINT Socket)
+{
+	return int(Socket-1);
+}
+
+////////////////////////////////////////////////////////////////////////////////
+static bool TcpSocketSetNonBlockingInternal(SOCKET Socket, bool bNonBlocking)
 {
 	unsigned long NonBlockingMode = !!bNonBlocking;
 	return ioctlsocket(Socket, FIONBIO, &NonBlockingMode) != SOCKET_ERROR;
 }
 
+////////////////////////////////////////////////////////////////////////////////
+bool TcpSocketSetNonBlocking(UPTRINT Socket, bool bNonBlocking)
+{
+	SOCKET Inner = (SOCKET) TcpSocketNative(Socket);
+	return TcpSocketSetNonBlockingInternal(Inner, bNonBlocking);
+}
+
 ////////////////////////////////////////////////////////////////////////////////
 UPTRINT TcpSocketConnect(const ANSICHAR* Host, uint16 Port)
 {
@@ -145,7 +158,7 @@ UPTRINT TcpSocketConnect(const ANSICHAR* Host, uint16 Port)
 		return 0;
 	}
 
-	if (!TcpSocketSetNonBlocking(Socket, 0))
+	if (!TcpSocketSetNonBlockingInternal(Socket, 0))
 	{
 		FSetLastErrorScope _(GetLastError());
 		closesocket(Socket);
@@ -187,7 +200,7 @@ UPTRINT TcpSocketListen(uint16 Port)
 		return 0;
 	}
 
-	if (!TcpSocketSetNonBlocking(Socket, 1))
+	if (!TcpSocketSetNonBlockingInternal(Socket, 1))
 	{
 		FSetLastErrorScope _(GetLastError());
 		closesocket(Socket);
@@ -198,9 +211,9 @@ UPTRINT TcpSocketListen(uint16 Port)
 }
 
 ////////////////////////////////////////////////////////////////////////////////
-int32 TcpSocketAccept(UPTRINT Socket, UPTRINT& Out)
+int32 TcpSocketAccept(UPTRINT Socket, UPTRINT& Out, bool bAsNonBlocking /*= false*/)
 {
-	SOCKET Inner = Socket - 1;
+	SOCKET Inner = (SOCKET) TcpSocketNative(Socket);
 
 	Inner = accept(Inner, nullptr, nullptr);
 	if (Inner == INVALID_SOCKET)
@@ -208,7 +221,7 @@ int32 TcpSocketAccept(UPTRINT Socket, UPTRINT& Out)
 		return (WSAGetLastError() == WSAEWOULDBLOCK) - 1; // 0 if would block else -1
 	}
 
-	if (!TcpSocketSetNonBlocking(Inner, 0))
+	if (!TcpSocketSetNonBlockingInternal(Inner, bAsNonBlocking))
 	{
 		FSetLastErrorScope _(GetLastError());
 		closesocket(Inner);
@@ -222,13 +235,30 @@ int32 TcpSocketAccept(UPTRINT Socket, UPTRINT& Out)
 ////////////////////////////////////////////////////////////////////////////////
 bool TcpSocketHasData(UPTRINT Socket)
 {
-	SOCKET Inner = Socket - 1;
+	SOCKET Inner = (SOCKET) TcpSocketNative(Socket);
 	fd_set FdSet = { 1, { Inner }, };
 	TIMEVAL TimeVal = {};
 	return (select(0, &FdSet, nullptr, nullptr, &TimeVal) != 0);
 }
 
 
+////////////////////////////////////////////////////////////////////////////////
+bool TcpSocketPeerName(UPTRINT Socket, char* Name, uint32 NameSize)
+{
+	sockaddr_storage Addr;
+	socklen_t AddrLen = sizeof(Addr);
+
+	getpeername(int(Socket-1), (sockaddr*) &Addr, &AddrLen);
+	if (Addr.ss_family != AF_INET)
+	{
+		return false;
+	}
+
+	sockaddr_in* Ip4Addr = (sockaddr_in*) &Addr;
+	inet_ntop(AF_INET, &Ip4Addr->sin_addr, Name, NameSize);
+
+	return true;
+}
 
 ////////////////////////////////////////////////////////////////////////////////
 bool IoWrite(UPTRINT Handle, const void* Data, uint32 Size)
diff --git a/Engine/Source/Runtime/TraceLog/Private/Trace/Platform.h b/Engine/Source/Runtime/TraceLog/Private/Trace/Platform.h
--- a/Engine/Source/Runtime/TraceLog/Private/Trace/Platform.h
+++ b/Engine/Source/Runtime/TraceLog/Private/Trace/Platform.h
@@ -34,8 +34,11 @@ TRACELOG_API uint64 TimeGetRelativeTimestamp();
 ////////////////////////////////////////////////////////////////////////////////
 UPTRINT	TcpSocketConnect(const ANSICHAR* Host, uint16 Port);
 UPTRINT	TcpSocketListen(uint16 Port);
-int32	TcpSocketAccept(UPTRINT Socket, UPTRINT& Out);
+int32	TcpSocketAccept(UPTRINT Socket, UPTRINT& Out, bool bAsNonBlocking = false);
 bool	TcpSocketHasData(UPTRINT Socket);
+int		TcpSocketNative(UPTRINT Socket);
+bool	TcpSocketPeerName(UPTRINT Socket, char* Name, uint32 NameSize);
+bool	TcpSocketSetNonBlocking(UPTRINT Socket, bool bNonBlocking);
 
 ////////////////////////////////////////////////////////////////////////////////
 int32	IoRead(UPTRINT Handle, void* Data, uint32 Size);
