[Insights] Secure socket transport support in TraceLog

- Added support for tracing through a secured SSL socket. This works slightly different from the standard socket connection, as the recieving side is expected to connect to the tracing application. The secure connection supports three modes:
* Self signed certificate and authorization code.
* Known certificate authority and authorization code.
* Pre shared key (PSK).
- Methods of making certificates and authorization codes available on the connecting application is not provided by TraceLog. This needs to be transfered out-of-band.
- Added functionality to extend a writing method using `IoUpdate` and `IoIsReady` functions. Changed how writer state is managed and synchronized internally. This means a pending data connection can now remain in pending state until it's ready. For secure socket connections this means it's ready once SSL handshake is complete and authorization token has been approved.
- Each connection emits logs for auditing using a unique identity string.
- Added OnClose callback.

#Synced-CL 49918891

diff --git a/Engine/Source/Programs/UnrealBuildTool/Configuration/Rules/ReadOnlyTargetRules.cs b/Engine/Source/Programs/UnrealBuildTool/Configuration/Rules/ReadOnlyTargetRules.cs
--- a/Engine/Source/Programs/UnrealBuildTool/Configuration/Rules/ReadOnlyTargetRules.cs
+++ b/Engine/Source/Programs/UnrealBuildTool/Configuration/Rules/ReadOnlyTargetRules.cs
@@ -200,6 +200,8 @@ namespace UnrealBuildTool
 
 		public bool bForceEnableTrace => Inner.bForceEnableTrace;
 
+		public bool bEnableTraceSecureConnection => Inner.bEnableTraceSecureConnection;
+
 		public bool bCompileAgainstEditor => Inner.bCompileAgainstEditor;
 
 		public bool bCompileRecast => Inner.bCompileRecast;
diff --git a/Engine/Source/Programs/UnrealBuildTool/Configuration/Rules/TargetRules.cs b/Engine/Source/Programs/UnrealBuildTool/Configuration/Rules/TargetRules.cs
--- a/Engine/Source/Programs/UnrealBuildTool/Configuration/Rules/TargetRules.cs
+++ b/Engine/Source/Programs/UnrealBuildTool/Configuration/Rules/TargetRules.cs
@@ -1351,6 +1351,25 @@ namespace UnrealBuildTool
 		}
 		private bool? bForceEnableTracePrivate;
 
+		/// <summary>
+		/// Enables secure socket connection support for Trace.
+		/// </summary>
+		[RequiresUniqueBuildEnvironment]
+		public virtual bool bEnableTraceSecureConnection
+		{
+			get 
+			{
+				if (bEnableTraceSecureConnectionPrivate.HasValue)
+				{
+					return bEnableTraceSecureConnectionPrivate.Value;
+				}
+
+				return false;
+			}
+			set => bEnableTraceSecureConnectionPrivate = value;
+		}
+		private bool? bEnableTraceSecureConnectionPrivate;
+
 		/// <summary>
 		/// Enabled for editor builds (TargetType.Editor). Can be overridden for programs (TargetType.Program) that would need to compile against editor code. Not available for other target types.
 		/// Mainly drives the value of WITH_EDITOR.
diff --git a/Engine/Source/Runtime/TraceLog/Private/Trace/Secure/CertUtils.cpp b/Engine/Source/Runtime/TraceLog/Private/Trace/Secure/CertUtils.cpp
new file mode 100644
--- /dev/null
+++ b/Engine/Source/Runtime/TraceLog/Private/Trace/Secure/CertUtils.cpp
@@ -0,0 +1,259 @@
+﻿// Copyright Epic Games, Inc. All Rights Reserved.
+#include "CertUtils.h"
+#include "Trace/Trace.h"
+#include "Trace/Message.h"
+
+#if UE_TRACE_MINIMAL_ENABLED && TRACE_PRIVATE_ALLOW_SECURE_TRACING 
+
+#include "openssl/rand.h"
+#include <cwchar>
+
+namespace UE::Trace::Private
+{
+
+////////////////////////////////////////////////////////////////////////////////
+void* Writer_MemoryAllocate(SIZE_T, uint32);
+void Writer_MemoryFree(void*, uint32);
+
+////////////////////////////////////////////////////////////////////////////////
+bool GeneratePrivateKey(EVP_PKEY* PrivateKey)
+{
+	if (!PrivateKey)
+	{
+		return false;
+	}
+
+	struct FScopedBigNum
+	{
+		FScopedBigNum()
+		{
+			BigNumPtr = BN_new();
+		}
+
+		~FScopedBigNum()
+		{
+			if (BigNumPtr)
+			{
+				BN_free(BigNumPtr);
+			}
+		}
+
+		BIGNUM* BigNumPtr = nullptr;
+	};
+
+	RSA* Rsa = RSA_new();
+	if (!Rsa)
+	{
+		return false;
+	}
+
+	FScopedBigNum ScopedBigNum;
+
+	if (!ScopedBigNum.BigNumPtr)
+	{
+		RSA_free(Rsa);
+		return false;
+	}
+
+	if (!BN_set_word(ScopedBigNum.BigNumPtr, RSA_F4))
+	{
+		RSA_free(Rsa);
+		return false;
+	}
+
+	constexpr int32 Bits = 2048;
+
+	if (!RSA_generate_key_ex(Rsa, Bits, ScopedBigNum.BigNumPtr, nullptr))
+	{
+		RSA_free(Rsa);
+		return false;
+	}
+
+	if (!EVP_PKEY_assign_RSA(PrivateKey, Rsa))
+	{
+		RSA_free(Rsa);
+		return false;
+	}
+
+	return true;
+}
+
+////////////////////////////////////////////////////////////////////////////////
+void AddSubjectNameEntryToCertificate(X509_NAME* SubjectName, const char* FiledType, const char* Value)
+{
+	if (SubjectName)
+	{	
+		constexpr int32 Len = -1; //Make OpenSSL calc len
+		constexpr int32 Loc = -1;
+		constexpr int32 Set = 0;
+		X509_NAME_add_entry_by_txt(
+			SubjectName, 
+			reinterpret_cast<const char*>(FiledType), 
+			MBSTRING_UTF8, 
+			reinterpret_cast<const unsigned char*>(Value),
+			Len,
+			Loc, 
+			Set
+		);
+	}
+}
+
+////////////////////////////////////////////////////////////////////////////////
+bool GenerateCertificate(EVP_PKEY* PrivateKey, X509* x509, const FCertificateOptions& InCertOptions)
+{
+	if (!PrivateKey)
+	{
+		return false;
+	}
+
+	if (!x509)
+	{
+		return false;
+	}
+
+	uint32 SerialNumber;
+	if (RAND_bytes(reinterpret_cast<unsigned char*>(&SerialNumber), sizeof(SerialNumber)) != 1)
+	{
+		UE_TRACE_MESSAGE_F(SecureSocketError, "[%s] Failed to generate serial number.", __func__);
+		return false;
+	}
+
+	SerialNumber = SerialNumber ? SerialNumber : 1;
+
+	if (!ASN1_INTEGER_set(X509_get_serialNumber(x509), SerialNumber))
+	{
+		UE_TRACE_MESSAGE_F(SecureSocketError, "[%s] Failed to generate assign number.", __func__);
+		return false;
+	}
+
+
+	// Set the certificate version to 3 (by default it is set to 1). V3 current version of the X.509 standard 
+	// X509_VERSION_3 is equals to 2
+	constexpr long CertificateVersion3 = 2;
+	if (!X509_set_version(x509, CertificateVersion3))
+	{
+		UE_TRACE_MESSAGE_F(SecureSocketError, "[%s] Failed to set the certificate version.", __func__);
+		return false;
+	}
+
+	// This certificate is valid from now until the specified duration. By default, the certificate will only last 1 hs
+	X509_gmtime_adj(X509_get_notBefore(x509), 0);
+	X509_gmtime_adj(X509_get_notAfter(x509), InCertOptions.DurationInSeconds);
+
+	// Set the public key for our certificate.
+	if (!X509_set_pubkey(x509, PrivateKey))
+	{
+		return false;
+	}
+
+	X509_NAME* SubjectName = X509_get_subject_name(x509);
+	if (!SubjectName)
+	{
+		UE_TRACE_MESSAGE_F(SecureSocketError, "[%s] Failed to get subject name.", __func__);
+		return false;
+	}
+
+	AddSubjectNameEntryToCertificate(SubjectName, "C", InCertOptions.Country);
+	AddSubjectNameEntryToCertificate(SubjectName, "O", InCertOptions.Organization);
+	AddSubjectNameEntryToCertificate(SubjectName, "CN", InCertOptions.CommonName);
+
+	if (!X509_set_issuer_name(x509, SubjectName))
+	{
+		UE_TRACE_MESSAGE_F(SecureSocketError, "[%s] Failed to set issuer name.", __func__);
+		return false;
+	}
+
+	if (!X509_sign(x509, PrivateKey, EVP_sha256()))
+	{
+		UE_TRACE_MESSAGE_F(SecureSocketError, "[%s] Failed to sign generated certificate.", __func__);
+		return false;
+	}
+
+	return true;
+}
+
+////////////////////////////////////////////////////////////////////////////////
+bool X509ToDER(X509* x509, uint8** OutEncoded, int32& OutEncodedSize)
+{
+	if (!x509 || !OutEncoded)
+	{
+		UE_TRACE_MESSAGE_F(SecureSocketError, "[%s] Failed to encode x509 certificate. No ptr provided", __func__);
+		return false;
+	}
+
+	const int32 ExpectedLength = i2d_X509(x509, nullptr);
+	if (ExpectedLength < 0)
+	{
+		UE_TRACE_MESSAGE_F(SecureSocketError, "[%s] Failed to encode x509 certificate. Invalid estimated length", __func__);
+		return false;
+	}
+
+	uint8* EncodedCertAlloc = (uint8*) Writer_MemoryAllocate(ExpectedLength, sizeof(void*));
+	uint8* EncodedCert = EncodedCertAlloc;
+	int32 FinalLength = i2d_X509(x509, &EncodedCert);
+	if (FinalLength != ExpectedLength)
+	{
+		Writer_MemoryFree(EncodedCertAlloc, ExpectedLength);
+		UE_TRACE_MESSAGE_F(SecureSocketError, "[%s] Failed to encode x509 certificate. Encode function returned an unexpected length | Expected [%d] | Result [%d]", __func__, ExpectedLength, FinalLength);
+		return false;
+	}
+
+	*OutEncoded = EncodedCertAlloc;
+	OutEncodedSize = FinalLength;
+	return true;
+}
+
+////////////////////////////////////////////////////////////////////////////////
+bool DERToX509(const uint8* EncodedCert, size_t EncodedSize, X509** OutCert)
+{
+	if (!EncodedCert || !EncodedSize || !OutCert)
+	{
+		UE_TRACE_MESSAGE_F(SecureSocketError, "[%s] Failed to decode x509 certificate. Invalid argument", __func__);
+		return false;
+	}
+
+	X509* DecodedCert = d2i_X509(nullptr, &EncodedCert, EncodedSize);
+
+	if (!DecodedCert)
+	{
+		UE_TRACE_MESSAGE_F(SecureSocketError, "[%s] Failed to decode encode x509 certificate.", __func__);
+		return false;
+	}
+
+	*OutCert = DecodedCert;
+	return true;
+}
+
+////////////////////////////////////////////////////////////////////////////////
+bool CreateSelfSignedCert(const FCertificateOptions& InCertOptions, X509** OutCert, EVP_PKEY** OutPrivateKey)
+{
+	if (!OutCert || !OutPrivateKey)
+	{
+		return false;
+	}
+	EVP_PKEY* PrivateKey = EVP_PKEY_new();
+	if (!GeneratePrivateKey(PrivateKey))
+	{
+		UE_TRACE_MESSAGE_F(SecureSocketError, "[%s] Failed to generate private key.", __func__);
+		EVP_PKEY_free(PrivateKey);
+		return false;
+	}
+
+	X509* x509 = X509_new();
+	if (!GenerateCertificate(PrivateKey, x509, InCertOptions))
+	{
+		UE_TRACE_MESSAGE_F(SecureSocketError, "[%s] Failed to generate certificate.", __func__);
+
+		EVP_PKEY_free(PrivateKey);
+		X509_free(x509);
+		return false;
+	}
+
+	*OutCert = x509;
+	*OutPrivateKey = PrivateKey;
+	return true;
+}
+
+} // namespace UE::Trace::Private
+
+#endif // UE_TRACE_MINIMAL_ENABLED && TRACE_PRIVATE_ALLOW_SECURE_TRACING  
diff --git a/Engine/Source/Runtime/TraceLog/Private/Trace/Secure/CertUtils.h b/Engine/Source/Runtime/TraceLog/Private/Trace/Secure/CertUtils.h
new file mode 100644
--- /dev/null
+++ b/Engine/Source/Runtime/TraceLog/Private/Trace/Secure/CertUtils.h
@@ -0,0 +1,76 @@
+﻿// Copyright Epic Games, Inc. All Rights Reserved.
+#pragma once
+
+#include "Trace/Config.h"
+#include "Trace/Platform.h"
+
+#if UE_TRACE_MINIMAL_ENABLED && TRACE_PRIVATE_ALLOW_SECURE_TRACING 
+
+#include "openssl/x509.h"
+
+namespace UE::Trace::Private
+{
+
+//todo: Make required configuration option when secure connection is enabled
+#define UE_TRACE_SECURESOCKET_DEFAULT_CMNNAME "Unreal Engine Trace"
+#define UE_TRACE_SECURESOCKET_DEFAULT_ORG "Epic Games Inc"
+#define UE_TRACE_SECURESOCKET_DEFAULT_COUNTRY "US"
+
+/**
+ * Generate a 2048-bit RSA key.
+ *
+ * @param PrivateKey The private key struct
+ */
+bool GeneratePrivateKey(EVP_PKEY* PrivateKey);
+
+/** 
+ * Options used to create an SSL Certificate. UTF-8 string
+ * are supported.
+ */
+struct FCertificateOptions
+{
+	const char* CommonName = UE_TRACE_SECURESOCKET_DEFAULT_CMNNAME;
+	const char* Organization = UE_TRACE_SECURESOCKET_DEFAULT_ORG;
+	const char* Country = UE_TRACE_SECURESOCKET_DEFAULT_COUNTRY;
+	long DurationInSeconds = 3600L;
+};
+
+/**
+ * Generate self-signed x509 certificate.
+ *
+ * @param PrivateKey The private key struct
+ * @param x509 The certificate struct
+ * @return True if successful, false otherwise
+ */
+bool GenerateCertificate(EVP_PKEY* PrivateKey, X509* x509, const FCertificateOptions& InCertOptions);
+
+/**
+ * Encodes the provided certificate in DER format
+ * @param x509 Certificate to Encode
+ * @param OutEncoded A DER encode SSL Certificate
+ * @param OutEncodedSize Size of the certificate
+ * @return True if successful, false otherwise
+ */
+bool X509ToDER(X509* x509, uint8** OutEncoded, int32& OutEncodedSize);
+
+/**
+ * Decodes a DER formatted buffer into a x509 certificate instance
+ * @param EncodedCert buffer containing the certificate encoded data
+ * @param EncodedSize size of the byffer
+ * @param OutCert A pointer to the generated X509 certificate instance
+ * @return True if successful, false otherwise
+ */
+bool DERToX509(const uint8* EncodedCert, size_t EncodedSize, X509** OutCert);
+
+/**
+ * Creates a self signed certificate based on input options.
+ * @param InCertOptions Options used to create the certificate
+ * @param OutCert Generated certificate
+ * @param OutPrivateKey Generated private key
+ * @return True if successful, false otherwise
+ */
+bool CreateSelfSignedCert(const FCertificateOptions& InCertOptions, X509** OutCert, EVP_PKEY** OutPrivateKey);
+
+} // namespace UE::Trace::Private
+
+#endif // UE_TRACE_MINIMAL_ENABLED && TRACE_PRIVATE_ALLOW_SECURE_TRACING 
diff --git a/Engine/Source/Runtime/TraceLog/Private/Trace/Secure/Secure.cpp b/Engine/Source/Runtime/TraceLog/Private/Trace/Secure/Secure.cpp
new file mode 100644
--- /dev/null
+++ b/Engine/Source/Runtime/TraceLog/Private/Trace/Secure/Secure.cpp
@@ -0,0 +1,992 @@
+// Copyright Epic Games, Inc. All Rights Reserved.
+#include "Trace/Config.h"
+
+#if UE_TRACE_MINIMAL_ENABLED && TRACE_PRIVATE_ALLOW_SECURE_TRACING 
+
+#include "Trace/Detail/Atomic.h"
+#include "Trace/Message.h"
+#include "Trace/Secure/CertUtils.h"
+#include "Trace/Trace.h"
+
+#include <openssl/ssl.h>
+#include <openssl/err.h>
+
+/**
+ * Enable extra verbose SSL debug logging
+ */
+#if !defined(TRACE_PRIVATE_SSL_DEBUG_LOGGING) 
+	#define TRACE_PRIVATE_SSL_DEBUG_LOGGING 0
+#endif
+
+
+namespace UE {
+namespace Trace {
+namespace Private {
+
+////////////////////////////////////////////////////////////////////////////////
+void Writer_InternalInitialize();
+bool Writer_SetPendingHandle(UPTRINT, IoWriteFunc, IoCloseFunc, IoUpdateFunc, IoIsReadyFunc);
+void* Writer_MemoryAllocate(SIZE_T, uint32);
+void Writer_MemoryFree(void*, uint32);
+UPTRINT Writer_PackSendFlags(UPTRINT, uint32);
+
+////////////////////////////////////////////////////////////////////////////////
+#define UE_TRACE_LOG_IDENT(level, fmt, ident)\
+	UE_TRACE_MESSAGE_F(level, "[S%p]" fmt, ident)
+
+////////////////////////////////////////////////////////////////////////////////
+#define UE_TRACE_LOG_IDENT_F(level, fmt, ident, ...)\
+	UE_TRACE_MESSAGE_F(level, "[S%p]" fmt, ident, __VA_ARGS__)
+
+////////////////////////////////////////////////////////////////////////////////
+#define UE_TRACE_ERROR_IDENT(fmt, ident)\
+	UE_TRACE_MESSAGE_F(SecureSocketError, "[S%p]" fmt, ident)
+
+////////////////////////////////////////////////////////////////////////////////
+#define UE_TRACE_ERROR_IDENT_F(fmt, ident, ...) \
+	UE_TRACE_MESSAGE_F(SecureSocketError, "[S%p]" fmt, ident, __VA_ARGS__)
+
+////////////////////////////////////////////////////////////////////////////////
+#define UE_TRACE_OPENSSL_ERROR_IDENT(message, ident) \
+	{ \
+		while(int ErrorCode = ERR_get_error())\
+		{\
+			char SSLErrBuf[256]; \
+			ERR_error_string_n(ErrorCode, SSLErrBuf, 256); \
+			UE_TRACE_MESSAGE_F(SecureSocketError, "[S%p]" message ": %s", ident, SSLErrBuf);\
+		}\
+	}\
+
+////////////////////////////////////////////////////////////////////////////////
+#define UE_TRACE_OPENSSL_ERROR(message) \
+	{ \
+		while(int ErrorCode = ERR_get_error())\
+		{\
+			char SSLErrBuf[256]; \
+			ERR_error_string_n(ErrorCode, SSLErrBuf, 256); \
+			UE_TRACE_MESSAGE_F(SecureSocketError, message ": %s", SSLErrBuf);\
+		}\
+	}\
+
+////////////////////////////////////////////////////////////////////////////////
+class FSSLContext 
+{
+	public:
+		FSSLContext() {}
+		~FSSLContext() {}
+		bool Initialize(const FSecureTraceSettings& Options);
+		void Teardown();
+		bool IsNonBlocking() const { return bNonBlocking; }
+
+		bool bNonBlocking = false;
+		uint16 Port;
+		SSL_CTX* Context = nullptr;
+		uint8* CertificateDER = nullptr;
+		int32 CertificateDERSize = 0;
+		ESecureTraceMode Mode;
+
+	private:
+		bool InitCertificates(const FSecureTraceSettings& Options);
+		bool InitPSK(const FSecureTraceSettings& Options);
+		void EnableDebugHooks();
+		void DebugDumpDER(const char* Filename);
+
+		X509* CertificateX509 = nullptr;
+		EVP_PKEY* PrivateKey = nullptr;
+};
+
+////////////////////////////////////////////////////////////////////////////////
+class FSSLConnection
+{
+	public:
+
+		FSSLConnection(const FSSLContext& InContext, const FSecureTraceOptions& Options);
+		FSSLConnection() = delete;
+		FSSLConnection(FSSLConnection&) = delete;
+		~FSSLConnection();
+
+		bool Write(const void* Data, uint32 Size);
+		void Close();
+		bool Update();
+		bool IsOpen() const;
+		bool Listen(uint16 Port);
+		static uint32 OnPSK(SSL *SSL, const char *Identity, unsigned char *PSK, unsigned int MaxPSKLen);
+		
+	private:
+		enum EState : uint8 {
+			None = 0,
+			Closed,
+			Listening,
+			Handshake,
+			Authenticating,
+			Open,
+			Error = 0xff
+		};
+
+		enum EMode : uint8 {
+			Undefined		= 0,
+			Send			= 1 << 1,
+			Receive			= 1 << 2,
+			Authenticated	= 1 << 3,
+			NonBlocking		= 1 << 4,
+		};
+
+		bool CheckForTimeout(const char* Reason);
+
+		uint8* AuthRecvBuffer;
+		const uint8* AuthToken;
+		char* Identity;
+		uint32 IdentitySize;
+		uint32 AuthSize;
+		uint32 AuthRead;
+		const FSSLContext& Context;
+		SSL* Ssl;
+		BIO* Bio;
+		EMode Mode;
+		EState State;
+		uint64 WaitStartCycles;
+		uint32 TimeoutMs;
+
+		UPTRINT ListenHandle;
+		UPTRINT SocketHandle;
+};
+
+
+////////////////////////////////////////////////////////////////////////////////
+void FSSLContext::Teardown()
+{
+	Writer_MemoryFree(CertificateDER, CertificateDERSize);
+	if (CertificateX509)
+	{
+		X509_free(CertificateX509);
+	}
+	if (PrivateKey)
+	{
+		EVP_PKEY_free(PrivateKey);
+	}
+	if (Context)
+	{
+		SSL_CTX_free(Context);
+	}
+}
+
+bool FSSLContext::Initialize(const FSecureTraceSettings& Options)
+{
+	Mode = Options.Mode;
+	Port = Options.Port;
+
+	{
+		uint64 OpenSSLOptions = OPENSSL_INIT_NO_ATEXIT;
+		OPENSSL_init_ssl(OpenSSLOptions, nullptr);
+	}
+
+	auto* Method = TLSv1_2_server_method();
+	Context = SSL_CTX_new(Method);
+
+	if (!Context)
+	{
+		UE_TRACE_OPENSSL_ERROR("Failed to create SSL context");
+		return false;
+	}
+
+	EnableDebugHooks();
+
+	// const char* ModeNames[] = {"Self signed", "Certificate file", "Pre shared keys"};
+	// UE_TRACE_MESSAGE_F(VeryVerbose, "SSLContext using mode %s", ModeNames[(uint8)Mode]);
+
+	SSL_CTX_set_min_proto_version(Context, TLS1_2_VERSION);
+	SSL_CTX_set_max_proto_version(Context, TLS1_2_VERSION);
+	SSL_CTX_set_options(Context, SSL_OP_NO_TLSv1_3);
+
+	// Setup certificates if those modes are requested
+	if (Options.Mode == ESecureTraceMode::SelfSigned || Options.Mode == ESecureTraceMode::CertificateFile)
+	{
+		if (!InitCertificates(Options))
+		{
+			return false;
+		}
+	}
+	else
+	{
+		if (!InitPSK(Options))
+		{
+			return false;
+		}
+	}
+
+	return true;
+}
+
+
+bool FSSLContext::InitPSK(const FSecureTraceSettings& Options)
+{
+	SSL_CTX_set_options(Context, SSL_OP_CIPHER_SERVER_PREFERENCE);
+	SSL_CTX_set_psk_server_callback(Context, &FSSLConnection::OnPSK);
+
+	static const char* CipherList = "ECDHE-PSK-CHACHA20-POLY1305";
+
+	if (!SSL_CTX_set_cipher_list(Context, CipherList))
+	{
+		UE_TRACE_OPENSSL_ERROR("Failed to set cipher suites");
+		return false;
+	}
+	
+	static const char* GroupLists = "P-384:P-256";
+
+	if (!SSL_CTX_set1_groups_list(Context, GroupLists))
+	{
+		UE_TRACE_OPENSSL_ERROR("Failed to set group list");
+		return false;
+	}
+
+	// UE_TRACE_MESSAGE_F(VeryVerbose, "Cipher list: '%s', groups list: '%s'", CipherList, GroupLists);
+	return true;
+}
+
+bool FSSLContext::InitCertificates(const FSecureTraceSettings& Options)
+{
+	// Create certificate or use provided
+	if (Options.Mode == ESecureTraceMode::SelfSigned)
+	{
+		// Create a certificate with the default named fields 
+		FCertificateOptions CertCreationOptions;
+		if (!CreateSelfSignedCert(CertCreationOptions, &CertificateX509, &PrivateKey))
+		{
+			return false;
+		}
+
+		if (!SSL_CTX_use_certificate(Context, CertificateX509))
+		{
+			UE_TRACE_OPENSSL_ERROR("Failed to load SSL certificate.");
+			return false;
+		}
+
+		if (!SSL_CTX_use_PrivateKey(Context, PrivateKey))
+		{
+			UE_TRACE_OPENSSL_ERROR("Failed to load private key.");
+			return false;
+		}
+
+		if (!X509ToDER(CertificateX509, &CertificateDER, CertificateDERSize))
+		{
+			UE_TRACE_OPENSSL_ERROR("Failed to decode certificate");
+			return false;
+		}
+
+		// Uncomment to dump the binary into a file
+		// DebugDumpDER("ServerCertificate.der");
+	}
+	else if (Options.Mode == ESecureTraceMode::CertificateFile)
+	{
+		if (!SSL_CTX_use_certificate_file(Context, Options.CertificatePath, SSL_FILETYPE_PEM))
+		{
+			// Don't log filepath to avoid leaking location
+			UE_TRACE_OPENSSL_ERROR("Failed to load SSL certificate from file.");
+			return false;
+		}
+
+		if (!SSL_CTX_use_PrivateKey_file(Context, Options.PrivateKeyPath, SSL_FILETYPE_PEM))
+		{
+			// Don't log filepath to avoid leaking location
+			UE_TRACE_OPENSSL_ERROR("Failed to load private key from file.");
+			return false;
+		}
+	}
+
+	// Verify we have matching certificate and key
+	if (!SSL_CTX_check_private_key(Context))
+	{
+		UE_TRACE_OPENSSL_ERROR("Could not verify private key");
+		return false;
+	}
+
+	return true;
+}
+
+void FSSLContext::EnableDebugHooks()
+{
+#if TRACE_PRIVATE_SSL_DEBUG_LOGGING 
+	SSL_load_error_strings();
+	ERR_load_BIO_strings();
+	OpenSSL_add_all_algorithms();
+
+	SSL_CTX_set_info_callback(
+	    Context,
+	    [](const SSL *Ssl, int Type, int Value)
+	    {
+		    const char *Direction = (Type & SSL_CB_READ) ? "READ" : "WRITE";
+
+		    // Check if the current state is Server processing Client Hello (after read)
+			OSSL_HANDSHAKE_STATE State = SSL_get_state(Ssl);
+		    if (State == TLS_ST_SR_CLNT_HELLO)
+		    {
+			    // Check if the Client Hello has been read (e.g., state > SSL_ST_BEFORE)
+			    if (SSL_get_current_cipher(Ssl) == NULL) // Before a cipher is chosen
+			    {
+				    // --- Log Client's Offered Ciphers ---
+				    STACK_OF(SSL_CIPHER) *ciphers = SSL_get_client_ciphers(Ssl);
+
+				    // Check if the stack is available (only available after Client Hello is processed)
+				    if (ciphers != NULL && (Type & SSL_CB_EXIT))
+				    {
+					    UE_TRACE_MESSAGE(VeryVerbose, "[TLS_INFO]|CLIENT_OFFERS|Ciphers offered by client:");
+
+					    for (int i = 0; i < sk_SSL_CIPHER_num(ciphers); i++)
+					    {
+						    const SSL_CIPHER *cipher = sk_SSL_CIPHER_value(ciphers, i);
+
+						    // Print the long name of the offered cipher
+						    UE_TRACE_MESSAGE_F(VeryVerbose, "[TLS_INFO]|OFFERED_CIPHER|%d: %s", i + 1,
+						                       SSL_CIPHER_get_name(cipher));
+					    }
+				    }
+			    }
+		    }
+
+		    // --- State and Handshake Events ---
+		    if (Type & (SSL_CB_LOOP | SSL_CB_EXIT))
+		    {
+			    // Log the state transition
+			    const char* StateStr = SSL_state_string_long(Ssl);
+			    // The message is constructed to include the state string.
+			    UE_TRACE_MESSAGE_F(VeryVerbose, "[TLS_INFO]|Dir[%s]|Event[STATE_CHANGE]|State[%s]|Result[%s]", Direction,
+			                       (StateStr ? StateStr : "N/A"), (Type & SSL_CB_EXIT) ? "EXIT" : "LOOP");
+			    return;
+		    }
+
+		    switch (Type)
+		    {
+		    case SSL_CB_HANDSHAKE_START:
+			    // Log the start of the handshake
+			    UE_TRACE_MESSAGE_F(VeryVerbose, "[TLS_INFO]|Dir[%s]|Event[HANDSHAKE_START]|Protocol[%s]", Direction,
+			                       SSL_get_version(Ssl) // Protocol version used for the Client/Server Hello
+			    );
+			    break;
+
+		    case SSL_CB_HANDSHAKE_DONE:
+			    UE_TRACE_MESSAGE_F(VeryVerbose, "[TLS_INFO]|Dir[%s]|Event[HANDSHAKE_DONE]|Protocol[%s]|Cipher[%s]", Direction,
+			                       SSL_get_version(Ssl),
+			                       SSL_get_cipher_name(Ssl) // The negotiated cipher 
+			    );
+			    break;
+
+		    case SSL_CB_ALERT:
+		    {
+			    // Log a received or sent alert
+			    const char* TypeStr = SSL_alert_type_string_long(Value);
+			    const char* DescStr = SSL_alert_desc_string_long(Value);
+
+			    // Log the alert with both the name and the raw code
+			    UE_TRACE_MESSAGE_F(VeryVerbose, "[TLS_INFO]|Dir[%s]|Event[ALERT]|Type[%s]|Desc[%s] (Code: %d)", Direction,
+			                       (TypeStr ? TypeStr : "N/A"), (DescStr ? DescStr : "N/A"), Value);
+			    break;
+		    }
+
+		    default:
+			    // Log any other unhandled event type
+			    UE_TRACE_MESSAGE_F(VeryVerbose, "[TLS_INFO]|Dir[%s]|Event[UNKNOWN]|Type[%d]|Value[%d]", Direction, Type, Value);
+			    break;
+		    }
+	    });
+
+	SSL_CTX_set_msg_callback(Context, [](int write_p, int version, int content_type, const void *buf, size_t len, SSL *ssl, void *arg) {
+		UE_TRACE_MESSAGE_F(VeryVerbose, 
+			"TLS Traffic [%s]|Type [%d]|Length [%zu]", 
+			write_p ? "In" : "Out", 
+			content_type, 
+			len
+		);
+	});
+#endif // WITH_SSL_DEBUG_LOGGING
+}
+
+void FSSLContext::DebugDumpDER(const char* Filename)
+{
+	auto FileHandle = FileOpen(Filename);
+	if (FileHandle)
+	{
+		const bool bSuccess = IoWrite(FileHandle, CertificateDER, CertificateDERSize);
+		if (!bSuccess)
+		{
+			UE_TRACE_MESSAGE_F(SecureSocketError, "Failed to dump server certificate to '%s'.", Filename);
+		}
+		else
+		{
+			UE_TRACE_MESSAGE_F(Verbose, "Successfully dumped certificate to %s (%d bytes)", Filename, CertificateDERSize);
+		}
+		IoClose(FileHandle);
+	}
+}
+
+////////////////////////////////////////////////////////////////////////////////
+
+FSSLConnection::FSSLConnection(const FSSLContext& InContext, const FSecureTraceOptions& Options)
+	: AuthRecvBuffer(nullptr)
+	, AuthToken(nullptr)
+	, Identity(nullptr)
+	, AuthRead(0)
+	, Context(InContext)
+	, Ssl(nullptr)
+	, Bio(nullptr)
+	, State(EState::Closed)
+	, WaitStartCycles(0)
+	, TimeoutMs(Options.TimeoutMs)
+	, ListenHandle(0)
+	, SocketHandle(0)
+{
+	const bool bAdditionalAuthentication = Options.AuthorizationTokenSize > 0 && InContext.Mode != ESecureTraceMode::PreSharedKeys;
+	Mode = (EMode) (
+		Send |
+		(Options.bCanReceive ? Receive : 0u) |
+		(bAdditionalAuthentication ? Authenticated : 0u) |
+		(InContext.IsNonBlocking() ? NonBlocking : 0u)
+	);
+
+	AuthSize = Options.AuthorizationTokenSize;
+	AuthToken = Options.AuthorizationToken;
+
+	if (Mode & Authenticated)
+	{
+		AuthRecvBuffer = (uint8*) Writer_MemoryAllocate(AuthSize, 8);
+	}
+
+	const uint32 IdentityLen = Options.Identity ? (strlen(Options.Identity) + 1) : 0;
+	IdentitySize = IdentityLen * sizeof(char);
+	if (IdentitySize)
+	{
+		Identity = (char*) Writer_MemoryAllocate(IdentitySize, 8);
+		strcpy_s(Identity, IdentityLen, Options.Identity);
+	}
+}
+
+FSSLConnection::~FSSLConnection()
+{
+	if (Ssl)
+	{
+		SSL_free(Ssl);
+	}
+	if (AuthRecvBuffer)
+	{
+		Writer_MemoryFree(AuthRecvBuffer, AuthSize);
+	}
+	if (Identity)
+	{
+		Writer_MemoryFree(Identity, IdentitySize);
+	}
+	if (ListenHandle)
+	{
+		IoClose(ListenHandle);
+	}
+	if (SocketHandle)
+	{
+		IoClose(SocketHandle);
+	}
+}
+
+bool FSSLConnection::Listen(uint16 ListenPort)
+{
+	ListenHandle = TcpSocketListen(ListenPort);
+	if (!ListenHandle)
+	{
+		UE_TRACE_ERROR_IDENT_F("Failed to create listen socket on port %u", this, ListenPort);
+		return false;
+	}
+	UE_TRACE_LOG_IDENT_F(Log, "Started listening on port %u.", this, ListenPort);
+	AtomicStoreRelaxed(&State, EState::Listening);
+	return true;
+}
+
+bool FSSLConnection::Update()
+{
+	switch(State)
+	{
+		case EState::Error:
+			return false;
+
+		case EState::Listening:
+			{
+				UPTRINT AcceptedHandle;
+				int Return = TcpSocketAccept(ListenHandle, AcceptedHandle, true);
+				if (Return <= 0)
+				{
+					if (Return == -1)
+					{
+						UE_TRACE_ERROR_IDENT("Failed to accept connection.", this);
+						AtomicStoreRelaxed(&State, EState::Error); 
+						return false;
+					}
+
+					if (CheckForTimeout("Timeout waiting for connection."))
+					{
+						return false;
+					}
+
+					return true; // No connection yet
+				}
+
+				// Log accepting connections
+				char AcceptingIp[256];
+				const bool bCouldResolveName = TcpSocketPeerName(AcceptedHandle, AcceptingIp, sizeof(AcceptingIp));
+				UE_TRACE_LOG_IDENT_F(Log, "Accepting connection from %s.", this, bCouldResolveName ? AcceptingIp : "unknown");
+
+				SocketHandle = AcceptedHandle;
+
+				// Create SSL state for this connection and setup buffered io
+				Ssl = SSL_new(Context.Context);
+				Bio = BIO_new_socket(TcpSocketNative(SocketHandle), 0); 
+
+				if (!Ssl || !Bio)
+				{
+					UE_TRACE_OPENSSL_ERROR_IDENT("Failed to create SSL/buffered io for connection", this);
+					AtomicStoreRelaxed(&State, EState::Error); 
+					SSL_free(Ssl);
+					Ssl = nullptr;
+					BIO_free(Bio); // Need to manually free the bio since it's not yet attached.
+					Bio = nullptr;
+					Close();
+					return false;
+				}
+
+				BIO_set_nbio(Bio, 1);
+				SSL_set_bio(Ssl, Bio, Bio);
+
+				if (Context.Mode == ESecureTraceMode::PreSharedKeys)
+				{
+					SSL_set_app_data(Ssl, this);
+				}
+
+				// Close listening socket
+				IoClose(ListenHandle);
+				ListenHandle = 0;
+
+				// Reset timeout for next stage
+				WaitStartCycles = 0;
+
+				AtomicStoreRelaxed(&State, EState::Handshake);
+			}
+			break;
+
+		case EState::Handshake:
+			{
+				const int Result = SSL_accept(Ssl);
+				if (Result <= 0)
+				{
+					int SSLError = SSL_get_error(Ssl, Result);
+					if (SSLError == SSL_ERROR_WANT_READ || SSLError == SSL_ERROR_WANT_WRITE)
+					{
+						if (CheckForTimeout("Timeout waiting for handshake"))
+						{
+							return false;
+						}
+						return true;
+					}
+					UE_TRACE_OPENSSL_ERROR_IDENT("Failed SSL handshake", this);
+					AtomicStoreRelaxed(&State, EState::Error); 
+					Close();
+					return false;
+				}
+
+				WaitStartCycles = 0;
+
+				UE_TRACE_LOG_IDENT(Verbose, "SSL handshake completeted.", this);
+				if ((Mode & Authenticated) != 0)
+				{
+					AtomicStoreRelaxed(&State, EState::Authenticating);
+				}
+				else
+				{
+					AtomicStoreRelaxed(&State, EState::Open);
+				}
+			}
+			break;
+
+		case EState::Authenticating:
+			{
+				const int BytesRead = SSL_read(Ssl, AuthRecvBuffer + AuthRead, AuthSize - AuthRead);
+				if (BytesRead < 0)
+				{
+					const int SslError = SSL_get_error(Ssl, BytesRead);
+					if (SslError == SSL_ERROR_WANT_READ || SslError == SSL_ERROR_WANT_WRITE)
+					{
+						if (CheckForTimeout("Timeout waiting for authentication"))
+						{
+							return false;
+						}
+						return true;
+					}
+					UE_TRACE_OPENSSL_ERROR_IDENT("Transfer error", this);
+					AtomicStoreRelaxed(&State, EState::Error); 
+					Close();
+					return false;
+				}
+				else if (BytesRead == 0)
+				{
+					UE_TRACE_LOG_IDENT(Log, "Connection closed by peer", this);
+					AtomicStoreRelaxed(&State, EState::Closed);
+					Close();
+					return false;
+				}
+				AuthRead += BytesRead;
+				if (AuthRead == AuthSize)
+				{
+					if (memcmp(AuthRecvBuffer, AuthToken, AuthSize) == 0)
+					{
+						UE_TRACE_LOG_IDENT(Log, "Authorisation token accepted.", this);
+						AtomicStoreRelaxed(&State, EState::Open);
+
+						// Disable non-blocking mode. This cannot happen if there is
+						// pending data, which we don't expect from the peer at this point.
+						if (SSL_pending(Ssl))
+						{
+							UE_TRACE_ERROR_IDENT("Unexpected pending data in connection", this);
+							return false;
+						}
+
+						if (!TcpSocketSetNonBlocking(SocketHandle, false)) 
+						{
+							UE_TRACE_ERROR_IDENT("Failed to reset non-blocking mode on socket", this);
+							return false;
+						}
+
+						// BIO control macros sometimes return the 0 event though
+						// no actual error occurred, so also check error stack
+						int BioRes = BIO_set_nbio(Bio, 0);
+						if (BioRes != 1 && ERR_peek_error() != 0) 
+						{
+							UE_TRACE_OPENSSL_ERROR_IDENT("Failed to disable non-blocking mode on bio", this);
+							return false;
+						}
+						
+						return true;
+					}
+					else
+					{
+						UE_TRACE_ERROR_IDENT("Authorisation token rejected.", this);
+						Close();
+						return false;
+					}
+				}
+			}
+			break;
+	}
+
+	return true;
+}
+
+bool FSSLConnection::Write(const void* Data, uint32 Size)
+{
+	if (!IsOpen()) 
+	{
+		return false;
+	}
+	
+	const int BytesWritten = SSL_write(Ssl, Data, Size);
+	if (BytesWritten != Size)
+	{
+		UE_TRACE_OPENSSL_ERROR_IDENT("Write error", this);
+	}
+	return BytesWritten == Size;
+}
+
+void FSSLConnection::Close()
+{
+	EState Previous = AtomicExchangeRelease(&State, EState::Closed);
+	if (Previous != EState::Closed)
+	{
+		UE_TRACE_LOG_IDENT(Log, "Connection closed", this);
+		IoClose(SocketHandle);
+		SocketHandle = 0;
+	}
+}
+
+inline bool FSSLConnection::IsOpen() const
+{
+	return AtomicLoadRelaxed<uint8>((uint8*)&State) == EState::Open;
+}
+
+bool FSSLConnection::CheckForTimeout(const char* Reason)
+{
+	if (!TimeoutMs)
+	{
+		return false;
+	}
+	const uint64 Ts = TimeGetRelativeTimestamp();
+	WaitStartCycles = WaitStartCycles ? WaitStartCycles : Ts;
+	uint32 WaitedMs = uint32(1000.f*(float(Ts - WaitStartCycles) / float(TimeGetFrequency())));
+	if (WaitedMs > TimeoutMs)
+	{
+		UE_TRACE_ERROR_IDENT_F(" %s", this, Reason);
+		AtomicStoreRelaxed(&State, EState::Error); 
+		return true;
+	}
+	return false;
+}
+
+uint32 FSSLConnection::OnPSK(SSL* Ssl, const char* Identity, unsigned char* PSK, unsigned int MaxPSKLen)
+{
+	FSSLConnection* Self = (FSSLConnection*)SSL_get_app_data(Ssl);
+	if (!Self)
+	{
+		UE_TRACE_MESSAGE(SecureSocketError, "Unable to get app data.");
+		return 0;
+	}
+
+	if (MaxPSKLen < Self->AuthSize)
+	{
+		UE_TRACE_OPENSSL_ERROR_IDENT("Target PSK buffer too small.", Self);
+		return 0;
+	}
+
+	if (!Identity)
+	{
+		UE_TRACE_ERROR_IDENT("Peer must provide identity.", Self);
+		return 0;
+	}
+
+	if (strcmp(Identity, Self->Identity) != 0)
+	{
+		UE_TRACE_ERROR_IDENT_F("Unknown identity '%s'", Self, Identity);
+		return 0;
+	}
+
+	// UE_TRACE_MESSAGE_F(Verbose, "Retrieving pre-shared key for identity '%s'", Identity);
+	memcpy(PSK, Self->AuthToken, Self->AuthSize);
+
+	return Self->AuthSize;
+}
+
+
+////////////////////////////////////////////////////////////////////////////////
+class FSecure 
+{
+	public:
+		static bool Initialize(const FSecureTraceSettings& Settings);
+		static bool IsInitialized();
+		static bool GetSelfSignedCert(uint8 **OutPtr, uint32 *OutSize);
+		static void Teardown();
+		static FSSLConnection* CreateConnection(const FSecureTraceOptions& Options);
+		static void RemoveConnection(FSSLConnection* Connection);
+
+	private:
+		static FSSLContext SSLContext;	
+		static FSSLConnection* SSLConnection;
+		static bool bIsInitialized;
+};
+
+bool FSecure::Initialize(const FSecureTraceSettings& Settings)
+{
+	static bool bInitializedInternal = [&Settings]() -> bool {
+		bool bSuccess = SSLContext.Initialize(Settings);
+		if (!bSuccess)
+		{
+			UE_TRACE_MESSAGE(SecureInitError, "Failed to initialize secure connections");
+		}
+		else
+		{
+			bIsInitialized = true;
+		}
+		return true;
+	}();
+	return bIsInitialized;
+}
+
+bool FSecure::IsInitialized()
+{
+	return bIsInitialized;
+}
+
+bool FSecure::GetSelfSignedCert(uint8 **OutPtr, uint32 *OutSize)
+{
+	if (SSLContext.CertificateDER)
+	{
+		*OutPtr = SSLContext.CertificateDER;
+		*OutSize = SSLContext.CertificateDERSize;
+		return true;
+	}
+	return false;
+}
+
+void FSecure::Teardown()
+{
+	if (SSLConnection)
+	{
+		SSLConnection->Close();
+		SSLConnection->~FSSLConnection();
+		Writer_MemoryFree(SSLConnection, sizeof(FSSLConnection));
+		SSLConnection = nullptr;
+	}
+	SSLContext.Teardown();
+}
+
+FSSLConnection* FSecure::CreateConnection(const FSecureTraceOptions& Options)
+{
+	if (SSLConnection)
+	{
+		UE_TRACE_MESSAGE(ConnectError, "A secure connection already exists"); 
+		return nullptr;
+	}
+
+	if (SSLContext.Mode == ESecureTraceMode::PreSharedKeys &&
+	    (Options.AuthorizationTokenSize == 0 || Options.AuthorizationToken == nullptr))
+	{
+		UE_TRACE_MESSAGE(ConnectError, "An authorization token must be provided in PreSharedKeys mode"); 
+		return nullptr;
+	}
+
+	if (SSLContext.Mode == ESecureTraceMode::PreSharedKeys && Options.Identity == nullptr)
+	{
+		UE_TRACE_MESSAGE(ConnectError, "An identity must be provided in PreSharedKeys mode"); 
+		return nullptr;
+	}
+
+	// For now we only support one secure connection
+	SSLConnection = (FSSLConnection*) Writer_MemoryAllocate(sizeof(FSSLConnection), 16);
+	SSLConnection = new(SSLConnection) FSSLConnection(SSLContext, Options); 
+	if (!SSLConnection->Listen(SSLContext.Port))
+	{
+		RemoveConnection(SSLConnection);
+		return nullptr;
+	}
+	return SSLConnection;
+}
+
+void FSecure::RemoveConnection(FSSLConnection* Connection)
+{
+	if (SSLConnection)
+	{
+		SSLConnection->Close();
+		SSLConnection->~FSSLConnection();
+		Writer_MemoryFree(SSLConnection, sizeof(FSSLConnection));
+		SSLConnection = nullptr;
+	}
+}
+
+FSSLContext FSecure::SSLContext;
+FSSLConnection* FSecure::SSLConnection;
+bool FSecure::bIsInitialized;
+
+////////////////////////////////////////////////////////////////////////////////
+bool SecureWrite(UPTRINT Handle, const void* Data, uint32 Size)
+{
+	if (!Handle)
+	{
+		return false;
+	}
+	FSSLConnection* Connection = (FSSLConnection*) Handle;
+	return Connection->Write(Data, Size);
+}
+
+////////////////////////////////////////////////////////////////////////////////
+void SecureClose(UPTRINT Handle)
+{
+	if (!Handle)
+	{
+		return;
+	}
+	FSSLConnection* Connection = (FSSLConnection*) Handle;
+	Connection->Close();
+	FSecure::RemoveConnection(Connection);
+}
+
+////////////////////////////////////////////////////////////////////////////////
+bool SecureUpdate(UPTRINT Handle)
+{
+	if (!Handle)
+	{
+		UE_TRACE_MESSAGE(SecureInitError, "No handle, closing connection");
+		return false;
+	}
+	FSSLConnection* Connection = (FSSLConnection*) Handle;
+	return Connection->Update();
+}
+
+////////////////////////////////////////////////////////////////////////////////
+bool SecureIsReady(UPTRINT Handle)
+{
+	if (!Handle)
+	{
+		return false;
+	}
+	FSSLConnection* Connection = (FSSLConnection*) Handle;
+	return Connection->IsOpen();
+}
+////////////////////////////////////////////////////////////////////////////////
+bool Writer_SecureInit(const FSecureTraceSettings& Settings)
+{
+	// Verify arguments
+	if (Settings.Mode == ESecureTraceMode::CertificateFile && (!Settings.CertificatePath || !Settings.PrivateKeyPath))
+	{
+		UE_TRACE_MESSAGE(SecureInitError, "CertificatePath and PrivateKeyPath needs to be set when using CertificateFile mode.");
+		return false;
+	}
+
+	return FSecure::Initialize(Settings);
+}
+
+////////////////////////////////////////////////////////////////////////////////
+bool Writer_SecureIsInit()
+{
+	return FSecure::IsInitialized();
+}
+
+bool Writer_SecureGetSelfSignedCert(uint8 **OutPtr, uint32 *OutSize)
+{
+	return FSecure::GetSelfSignedCert(OutPtr, OutSize);
+}
+
+////////////////////////////////////////////////////////////////////////////////
+bool Writer_SecureSendTo(const FSecureTraceOptions& Options, uint16 Flags)
+{
+	if (!FSecure::IsInitialized())
+	{
+		UE_TRACE_MESSAGE(ConnectError, "Secure connection has not been initialized, failed to create connection.");
+		return false;
+	}
+	FSSLConnection* NewConnection = FSecure::CreateConnection(Options);
+	if (!NewConnection)
+	{
+		UE_TRACE_MESSAGE(ConnectError, "Failed to create secure connection.");
+		return false;
+	}
+
+	UPTRINT DataHandle = (UPTRINT) NewConnection;
+	DataHandle = Writer_PackSendFlags(DataHandle, Flags);
+
+	return Writer_SetPendingHandle(
+		DataHandle, 
+		SecureWrite, 
+		SecureClose, 
+		SecureUpdate, 
+		SecureIsReady
+	);
+}
+
+} //namespace Private
+
+////////////////////////////////////////////////////////////////////////////////
+bool SecureIsInitialized()
+{
+	return Private::Writer_SecureIsInit();
+}
+
+////////////////////////////////////////////////////////////////////////////////
+bool SecureInitialize(const FSecureTraceSettings& Settings)
+{
+	return Private::Writer_SecureInit(Settings);
+}
+
+bool SecureGetSelfSignedCertificate(uint8 **OutPtr, uint32 *OutSize)
+{
+	return Private::Writer_SecureGetSelfSignedCert(OutPtr, OutSize);
+}
+
+////////////////////////////////////////////////////////////////////////////////
+bool SecureSendTo(const FSecureTraceOptions& Options, uint16 Flags)
+{
+	return Private::Writer_SecureSendTo(Options, Flags);
+}
+
+} //namespace Trace
+} //namespace UE
+
+#endif // UE_TRACE_MINIMAL_ENABLED
diff --git a/Engine/Source/Runtime/TraceLog/Private/Trace/Writer.cpp b/Engine/Source/Runtime/TraceLog/Private/Trace/Writer.cpp
--- a/Engine/Source/Runtime/TraceLog/Private/Trace/Writer.cpp
+++ b/Engine/Source/Runtime/TraceLog/Private/Trace/Writer.cpp
@@ -1,6 +1,7 @@
 // Copyright Epic Games, Inc. All Rights Reserved.
 
 #include "Trace/Config.h"
+#include "Trace/Lock.h"
 
 #if TRACE_PRIVATE_MINIMAL_ENABLED
 
@@ -65,6 +66,7 @@ static bool		Writer_SessionPrologue();
 void			Writer_FreeBlockListToPool(FWriteBuffer*, FWriteBuffer*);
 void			Writer_SetBlockPoolLimit(uint32);
 void			Writer_UnsetBlockPoolLimit();
+void			Writer_ClearPendingHandle();
 
 
 ////////////////////////////////////////////////////////////////////////////////
@@ -115,10 +117,13 @@ struct FWriterState
 {
 	IoWriteFunc* Write = nullptr;
 	IoCloseFunc* Close = nullptr;
+	IoUpdateFunc* Update = nullptr;
+	IoIsReadyFunc* IsReady = nullptr;
 };
 
 static FWriterState GWriterState;
 static FWriterState GPendingWriterState;
+static uint8 GPendingWriterLock;
 
 ////////////////////////////////////////////////////////////////////////////////
 // When a thread terminates we want to recover its trace buffer. To do this we
@@ -283,7 +288,11 @@ void* Writer_MemoryAllocate(SIZE_T Size, uint32 Alignment)
 #if defined(_MSC_VER)
 		Ret = _aligned_malloc(Size, Alignment);
 #elif (defined(__ANDROID_API__) && __ANDROID_API__ < 28) || defined(__APPLE__)
-		posix_memalign(&Ret, Alignment, Size);
+		int Result = posix_memalign(&Ret, Alignment, Size);
+		if (Result != 0)
+		{
+			UE_TRACE_MESSAGE_F(OOMFatal, "Failed to allocate %llu bytes with alignment %u: %d", uint64(Size), Alignment, Result);
+		}
 #else
 		Ret = aligned_alloc(Alignment, Size);
 #endif
@@ -642,6 +651,7 @@ static void Writer_DescribeAnnounce()
 static int8				GSyncPacketCountdown; // = 0
 static const int8		GNumSyncPackets = 3;
 static OnConnectFunc*	GOnConnection = nullptr;
+static OnCloseFunc*		GOnClose = nullptr;
 static OnUpdateFunc*	GOnUpdate = nullptr;
 static FTraceGuid		GSessionGuid; // = {0, 0, 0, 0};
 static FTraceGuid		GTraceGuid; // = {0, 0, 0, 0};
@@ -731,6 +741,10 @@ static void Writer_Close()
 	{
 		Writer_FlushSendBuffer();
 		GWriterState.Close(GDataHandle);
+		if (GOnClose)
+		{
+			GOnClose();
+		}
 	}
 
 	GDataHandle = 0;
@@ -741,12 +755,26 @@ static bool Writer_UpdateConnection()
 {
 	UPTRINT PendingDataHandle = AtomicLoadRelaxed(&GPendingDataHandle);
 
+	// Update active connection if an update function has been set
+	if (GDataHandle && GWriterState.Update != nullptr)
+	{
+		if (!GWriterState.Update(GDataHandle))
+		{
+			GWriterState.Close(GDataHandle);
+			GDataHandle = 0;
+		}
+	}
+
 	if (!PendingDataHandle)
 	{
 		return false;
 	}
 
 	FProfilerScope _(__func__);
+	
+	// Pending writer state is guaranteed to be written if the
+	// pending data handle is set. See Writer_SetPendingHandle
+	FWriterState PendingWriterState = GPendingWriterState;
 
 	// Is this a close request? So that we capture some of the events around
 	// the closure we will add some inertia before enacting the close.
@@ -765,10 +793,7 @@ static bool Writer_UpdateConnection()
 
 		return true;
 	}
-
-	FWriterState PendingWriterState = GPendingWriterState;
-	AtomicStoreRelaxed(&GPendingDataHandle, UPTRINT(0));
-
+	
 	// Extract send flags
 	uint32 SendFlags = uint32(PendingDataHandle >> 48ull);
 	PendingDataHandle &= 0x0000'ffff'ffff'ffffull;
@@ -776,10 +801,30 @@ static bool Writer_UpdateConnection()
 	// Reject the pending connection if we've already got a connection
 	if (GDataHandle)
 	{
-		GWriterState.Close(PendingDataHandle);
+		PendingWriterState.Close(PendingDataHandle);
+		Writer_ClearPendingHandle();
 		return false;
 	}
 
+	// Update the pending connection
+	if (PendingWriterState.Update != nullptr)
+	{
+		if(!PendingWriterState.Update(PendingDataHandle))
+		{
+			PendingWriterState.Close(PendingDataHandle);
+			Writer_ClearPendingHandle();
+			return false;
+		}
+	}
+	// Check if the pending connection is ready
+	if (PendingWriterState.IsReady != nullptr && !PendingWriterState.IsReady(PendingDataHandle))
+	{
+		return false;
+	}
+
+	// Now clear pending handle, we're committing to this as the new active connection
+	Writer_ClearPendingHandle();
+
 	// Generate Guid for new connection
 	Writer_CreateGuid(&GTraceGuid);
 
@@ -1145,6 +1190,7 @@ void Writer_Initialize(const FInitializeDesc& Desc)
 	// Store callback on connection
 	GOnConnection = Desc.OnConnectionFunc;
 	AtomicStoreRelaxed(&GOnUpdate, Desc.OnUpdateFunc);
+	GOnClose = Desc.OnCloseFunc;
 
 	// Allow the worker thread to start updating
 	AtomicStoreRelease(&GUpdateInProgress, uint32(0));
@@ -1166,7 +1212,7 @@ void Writer_Update()
 }
 
 ////////////////////////////////////////////////////////////////////////////////
-static UPTRINT Writer_PackSendFlags(UPTRINT DataHandle, uint32 Flags)
+UPTRINT Writer_PackSendFlags(UPTRINT DataHandle, uint32 Flags)
 {
 	// Passing ownership of IO to the worker thread via a single 64 bit value is
 	// convenient and saves a lot of machinery for something that mostly never
@@ -1184,6 +1230,30 @@ static UPTRINT Writer_PackSendFlags(UPTRINT DataHandle, uint32 Flags)
 	return DataHandle | (UPTRINT(Flags) << 48ull);
 }
 
+////////////////////////////////////////////////////////////////////////////////
+bool Writer_SetPendingHandle(UPTRINT NewHandle, IoWriteFunc Write, IoCloseFunc Close, IoUpdateFunc Update = nullptr, IoIsReadyFunc IsReady = nullptr)
+{
+	FLock Lock(&GPendingWriterLock);
+	if (AtomicLoadAcquire(&GPendingDataHandle))
+	{
+		// Some other thread has already set the pending handle
+		return false;
+	}
+	// Write the pending writer state first, then release the
+	// pending data handle to guarantee the state is available.
+	GPendingWriterState = {Write, Close, Update, IsReady};
+	AtomicStoreRelease(&GPendingDataHandle, NewHandle);
+	return true;
+}
+
+////////////////////////////////////////////////////////////////////////////////
+void Writer_ClearPendingHandle()
+{
+	FLock Lock(&GPendingWriterLock);
+	GPendingWriterState = {nullptr, nullptr, nullptr, nullptr};
+	AtomicStoreRelease(&GPendingDataHandle, UPTRINT(0));
+}
+
 ////////////////////////////////////////////////////////////////////////////////
 bool Writer_SendTo(const ANSICHAR* Host, uint32 Flags, uint32 Port)
 {
@@ -1210,10 +1280,7 @@ bool Writer_SendTo(const ANSICHAR* Host, uint32 Flags, uint32 Port)
 		return false;
 	}
 
-	GPendingWriterState = { &IoWrite, &IoClose };
-
-	AtomicStoreRelease(&GPendingDataHandle, DataHandle);
-	return true;
+	return Writer_SetPendingHandle(DataHandle, &IoWrite, &IoClose);
 #else
 	return false;
 #endif
@@ -1243,10 +1310,7 @@ bool Writer_WriteTo(const ANSICHAR* Path, uint32 Flags)
 		return false;
 	}
 
-	GPendingWriterState = { &IoWrite, &IoClose };
-
-	AtomicStoreRelease(&GPendingDataHandle, DataHandle);
-	return true;
+	return Writer_SetPendingHandle(DataHandle, &IoWrite, &IoClose);
 #else
 	return false;
 #endif
@@ -1269,10 +1333,7 @@ bool Writer_RelayTo(UPTRINT InHandle, IoWriteFunc WriteFunc, IoCloseFunc CloseFu
 		return false;
 	}
 
-	GPendingWriterState = { WriteFunc, CloseFunc };
-
-	AtomicStoreRelease(&GPendingDataHandle, DataHandle);
-	return true;
+	return Writer_SetPendingHandle(DataHandle, WriteFunc, CloseFunc);
 }
 
 ////////////////////////////////////////////////////////////////////////////////
@@ -1404,11 +1465,12 @@ bool Writer_WriteSnapshot(const FSnapshotTarget& Target)
 		TStashGlobal DataHandle(GDataHandle);
 		TStashGlobal StateWrite(GWriterState.Write);
 		TStashGlobal StateClose(GWriterState.Close);
-		TStashGlobal PendingDataHandle(GPendingDataHandle);
+		TStashGlobal StateUpdate(GWriterState.Update);
+		TStashGlobal StateIsReady(GWriterState.IsReady);
 		TStashGlobal SyncPacketCountdown(GSyncPacketCountdown, GNumSyncPackets);
 		TStashGlobal TraceStatistics(GTraceStatistics);
 
-		GWriterState = { &IoWrite, &IoClose };
+		GWriterState = { &IoWrite, &IoClose, nullptr, nullptr };
 
 		if (Target.Type == FSnapshotTarget::EType::FileTarget)
 		{
diff --git a/Engine/Source/Runtime/TraceLog/Public/Trace/Config.h b/Engine/Source/Runtime/TraceLog/Public/Trace/Config.h
--- a/Engine/Source/Runtime/TraceLog/Public/Trace/Config.h
+++ b/Engine/Source/Runtime/TraceLog/Public/Trace/Config.h
@@ -102,6 +102,16 @@ static_assert(UE_TRACE_MINIMAL_ENABLED == 0 || UE_TRACE_ENABLED_SHIPPING_EXPERIM
 #	define TRACE_PRIVATE_ALLOW_FILE UE_TRACE_ALLOW_FILE_TRACING
 #endif
 
+/**
+ * Control tracing to secure socket connections.
+ */
+#if !defined(UE_TRACE_ALLOW_SECURE_TRACING)
+#	define UE_TRACE_ALLOW_SECURE_TRACING 0
+#	define TRACE_PRIVATE_ALLOW_SECURE_TRACING 0
+#else
+#	define TRACE_PRIVATE_ALLOW_SECURE_TRACING UE_TRACE_ALLOW_SECURE_TRACING
+#endif
+
 /**
  * Control if important events are enabled
  */
diff --git a/Engine/Source/Runtime/TraceLog/Public/Trace/Trace.h b/Engine/Source/Runtime/TraceLog/Public/Trace/Trace.h
--- a/Engine/Source/Runtime/TraceLog/Public/Trace/Trace.h
+++ b/Engine/Source/Runtime/TraceLog/Public/Trace/Trace.h
@@ -97,6 +97,8 @@ TEventRef<IdType> MakeEventRef(IdType InId, uint32 InTypeId)
 enum class EMessageType : uint8
 {
 	Reserved			= 0,
+	VeryVerbose,
+	Verbose,
 	// Add to log
 	Log,
 	// For backwards compatibility
@@ -117,6 +119,8 @@ enum class EMessageType : uint8
 	WriterError,
 	CompressionError,
 	PanicError,
+	SecureInitError,
+	SecureSocketError,
 	// Fatal errors should cause application to stop
 	FatalStart			= 0x40,
 	GenericFatal,
@@ -136,6 +140,7 @@ struct FMessageEvent
 using OnMessageFunc = void(const FMessageEvent&);
 using OnConnectFunc = void(void);
 using OnUpdateFunc = void(void);
+using OnCloseFunc = void(void);
 using OnScopeBeginFunc = void(const ANSICHAR*);
 using OnScopeEndFunc = void(void);
 
@@ -151,6 +156,7 @@ struct FInitializeDesc
 	OnUpdateFunc*		OnUpdateFunc		= nullptr;
 	OnScopeBeginFunc*	OnScopeBeginFunc	= nullptr;
 	OnScopeEndFunc*		OnScopeEndFunc		= nullptr;
+	OnCloseFunc*			OnCloseFunc			= nullptr;
 };
 
 typedef uint32 FChannelId;
@@ -212,6 +218,22 @@ typedef bool		IoWriteFunc(UPTRINT Handle, const void* Data, uint32 Size);
  */
 typedef void		IoCloseFunc(UPTRINT Handle);	
 
+
+/**
+ * User defined update  callback.
+ * @param Handle User defined handle passed to the function
+ * @return True while connection is active, false if closed or error state occurred
+ */
+typedef bool		IoUpdateFunc(UPTRINT Handle);
+
+/**
+ * User defined is ready callback. Used to check if a pending connection is ready
+ * to be used.
+ * @param Handle User defined handle passed to the function
+ * @return True if ready, false if further waiting is needed.
+ */
+typedef bool		IoIsReadyFunc(UPTRINT Handle);
+
 struct FStatistics
 {
 	uint64	BytesSent				= 0;	// Bytes sent/written to the current endpoint
@@ -424,6 +446,115 @@ UE_TRACE_API FChannel* FindChannel(const TCHAR* ChannelName) UE_TRACE_IMPL(nullp
  */
 UE_TRACE_API FChannel* FindChannel(FChannelId ChannelId) UE_TRACE_IMPL(nullptr);
 
+
+#if TRACE_PRIVATE_ALLOW_SECURE_TRACING
+
+
+enum class ESecureTraceMode : uint8
+{
+	SelfSigned,
+	CertificateFile,
+	PreSharedKeys,
+};
+
+/**
+ * Options for initializing secure tracing. Supports three modes of operation 
+ * depending on what values are set.
+ * * External certificate + private key and optional authentication token.
+ * * Self signed certificate and optional authentication token.
+ * * Pre-shared keys
+ */
+struct FSecureTraceSettings 
+{
+	/**
+	 * Desired operational mode.
+	 */
+	ESecureTraceMode Mode = ESecureTraceMode::SelfSigned;
+
+	/**
+	 * Port used to listen for peer connections.
+	 */
+	uint16 Port = 0;
+
+	/**
+	 * Path to certificate file. Only used if Mode is CertificateFile.
+	 */
+	const ANSICHAR* CertificatePath = nullptr;
+
+	/**
+	 * Path to private key. Only used if Mode is CertificateFile.
+	 */
+	const ANSICHAR* PrivateKeyPath = nullptr;
+};
+
+/**
+ * Initialize secure connection component. This must happen exactly once before
+ * it's possible to call \ref SecureSendTo.
+ * @param Settings Settings to use for secure connections
+ */
+UE_TRACE_API bool SecureInitialize(const FSecureTraceSettings& Settings);
+
+/**
+ * Checks if secure connection component has been initialized.
+ * @return true if \ref SecureInitialize has already been called.
+ */
+UE_TRACE_API bool SecureIsInitialized();
+
+/**
+ * Retrieve the public certificate (DER format) when using a self-signed certificate.
+ * @param OutPtr If successful, the pointer to the certificate data. The pointer is valid for the remiander of the process lifetime.
+ * @param OutSize If successful, certificate size in bytes.
+ * @return True if initialized with self-signed certificate, and the certificate was successfully generated.
+ */
+UE_TRACE_API bool SecureGetSelfSignedCertificate(uint8** OutPtr, uint32* OutSize);
+
+/**
+ * Options for a secure trace connection. Allows to set a unique authorization token 
+ * and identity where applicable.
+ */
+struct FSecureTraceOptions
+{
+	/**
+	 * AuthorizationToken that peer will use to connect. Used as pre-shared keys if Mode is PreSharedKeys
+	 */
+	const uint8* AuthorizationToken = nullptr;
+
+	/**
+	 * Identity string expected. Used when Mode is PreSharedKeys. String does not need to outlive function call.
+	 */
+	const char* Identity = nullptr;
+
+	/**
+	 * Expected authorization token size in bytes. Must be set if AuthorizationToken is set.
+	 */
+	uint32 AuthorizationTokenSize = 0u;
+
+	/**
+	 * If the connection can receive data.
+	 */
+	bool bCanReceive = false;
+
+	/**
+	 * Timeout in ms for connect and handshakes. If peer doesn't connect
+	 * or supply handshake within this time connection is closed. Zero is
+	 * interpreted as no timeouts should be applied.
+	 */
+	uint32 TimeoutMs = 0u;
+};
+
+
+/**
+ * Setup output to a secure socket connection. This requires the receiver to connect to the tracing machine on the
+ * specified listening port. Additionally certificates to use and optional authorization tokens can be specified for
+ * the connection.
+ * @param Options Configuration options for the connection
+ * @param Flags Options for the connection
+ */
+UE_TRACE_API bool	SecureSendTo(const FSecureTraceOptions& Options, uint16 Flags = FSendFlags::None);
+
+
+#endif // TRACE_PRIVATE_ALLOW_SECURE_TRACING
+
 } // namespace Trace
 } // namespace UE
 
diff --git a/Engine/Source/Runtime/TraceLog/TraceLog.Build.cs b/Engine/Source/Runtime/TraceLog/TraceLog.Build.cs
--- a/Engine/Source/Runtime/TraceLog/TraceLog.Build.cs
+++ b/Engine/Source/Runtime/TraceLog/TraceLog.Build.cs
@@ -20,5 +20,11 @@ public class TraceLog : ModuleRules
 		}
 
 		PrivateDefinitions.Add("SUPPRESS_PER_MODULE_INLINE_FILE"); // This module does not use core's standard operator new/delete overloads
+
+		if (Target.bEnableTraceSecureConnection)
+		{
+			PublicDefinitions.Add("UE_TRACE_ALLOW_SECURE_TRACING=1");
+			PrivateDependencyModuleNames.Add("OpenSSL");
+		}
 	}
 }
