[Insights] Fix secure connection for Unix platforms.

Use the standard string copy for platform compatability and mute warnings.

#Synced-CL 50170004

diff --git a/Engine/Source/Runtime/TraceLog/Private/Trace/Secure/Secure.cpp b/Engine/Source/Runtime/TraceLog/Private/Trace/Secure/Secure.cpp
--- a/Engine/Source/Runtime/TraceLog/Private/Trace/Secure/Secure.cpp
+++ b/Engine/Source/Runtime/TraceLog/Private/Trace/Secure/Secure.cpp
@@ -455,7 +455,8 @@ FSSLConnection::FSSLConnection(const FSSLContext& InContext, const FSecureTraceO
 	if (IdentitySize)
 	{
 		Identity = (char*) Writer_MemoryAllocate(IdentitySize, 8);
-		strcpy_s(Identity, IdentityLen, Options.Identity);
+		strncpy(Identity, Options.Identity, IdentityLen);
+		Identity[IdentityLen-1] = '\0';
 	}
 }
 
@@ -840,7 +841,7 @@ FSSLConnection* FSecure::CreateConnection(const FSecureTraceOptions& Options)
 	}
 
 	// For now we only support one secure connection
-	SSLConnection = (FSSLConnection*) Writer_MemoryAllocate(sizeof(FSSLConnection), 16);
+	SSLConnection = (FSSLConnection*) Writer_MemoryAllocate(sizeof(FSSLConnection), alignof(FSSLConnection));
 	SSLConnection = new(SSLConnection) FSSLConnection(SSLContext, Options); 
 	if (!SSLConnection->Listen(SSLContext.Port))
 	{
diff --git a/Engine/Source/Runtime/TraceLog/TraceLog.Build.cs b/Engine/Source/Runtime/TraceLog/TraceLog.Build.cs
--- a/Engine/Source/Runtime/TraceLog/TraceLog.Build.cs
+++ b/Engine/Source/Runtime/TraceLog/TraceLog.Build.cs
@@ -26,5 +26,10 @@ public class TraceLog : ModuleRules
 			PublicDefinitions.Add("UE_TRACE_ALLOW_SECURE_TRACING=1");
 			PrivateDependencyModuleNames.Add("OpenSSL");
 		}
+
+		if (Target.Platform.IsInGroup(UnrealPlatformGroup.Windows))
+		{
+			PrivateDefinitions.Add("_CRT_SECURE_NO_WARNINGS");
+		}
 	}
 }
