[Insights] Make sure all paths that close connection goes through Writer_Close

#Synced-CL 50431656

diff --git a/Engine/Source/Runtime/TraceLog/Private/Trace/Writer.cpp b/Engine/Source/Runtime/TraceLog/Private/Trace/Writer.cpp
--- a/Engine/Source/Runtime/TraceLog/Private/Trace/Writer.cpp
+++ b/Engine/Source/Runtime/TraceLog/Private/Trace/Writer.cpp
@@ -67,6 +67,8 @@ void			Writer_FreeBlockListToPool(FWriteBuffer*, FWriteBuffer*);
 void			Writer_SetBlockPoolLimit(uint32);
 void			Writer_UnsetBlockPoolLimit();
 void			Writer_ClearPendingHandle();
+static void		Writer_Close();
+static void		Writer_CloseNoFlush();
 
 
 ////////////////////////////////////////////////////////////////////////////////
@@ -360,13 +362,12 @@ uint8* GSendBuffer; // = nullptr;
 uint8* GSendBufferCursor; // = nullptr;
 static bool Writer_FlushSendBuffer()
 {
-	if( GSendBufferCursor > GSendBuffer )
+	if( GDataHandle && GSendBufferCursor > GSendBuffer )
 	{
 		if (!GWriterState.Write(GDataHandle, GSendBuffer, GSendBufferCursor - GSendBuffer))
 		{
 			UE_TRACE_ERRORMESSAGE(WriteError, GetLastErrorCode());
-			GWriterState.Close(GDataHandle);
-			GDataHandle = 0;
+			Writer_CloseNoFlush(); // Don't flush or we'll recurse
 			return false;
 		}
 		GSendBufferCursor = GSendBuffer;
@@ -399,8 +400,7 @@ static void Writer_SendDataImplNoInstr(const void* Data, uint32 Size)
 	{
 		if (!GWriterState.Write(GDataHandle, Data, Size))
 		{
-			GWriterState.Close(GDataHandle);
-			GDataHandle = 0;
+			Writer_CloseNoFlush();
 		}
 	}
 	// Otherwise append to the buffer
@@ -413,8 +413,7 @@ static void Writer_SendDataImplNoInstr(const void* Data, uint32 Size)
 	if (!GWriterState.Write(GDataHandle, Data, Size))
 	{
 		UE_TRACE_ERRORMESSAGE(WriteError, GetLastErrorCode());
-		GWriterState.Close(GDataHandle);
-		GDataHandle = 0;
+		Writer_CloseNoFlush();
 	}
 #endif
 }
@@ -736,17 +735,22 @@ static void Writer_SendSync()
 
 ////////////////////////////////////////////////////////////////////////////////
 static void Writer_Close()
+{
+	Writer_FlushSendBuffer();
+	Writer_CloseNoFlush();
+}
+
+////////////////////////////////////////////////////////////////////////////////
+static void Writer_CloseNoFlush()
 {
 	if (GDataHandle)
 	{
-		Writer_FlushSendBuffer();
 		GWriterState.Close(GDataHandle);
 		if (GOnClose)
 		{
 			GOnClose();
 		}
 	}
-
 	GDataHandle = 0;
 }
 
@@ -760,8 +764,7 @@ static bool Writer_UpdateConnection()
 	{
 		if (!GWriterState.Update(GDataHandle))
 		{
-			GWriterState.Close(GDataHandle);
-			GDataHandle = 0;
+			Writer_CloseNoFlush();
 		}
 	}
 
@@ -915,8 +918,7 @@ static bool Writer_SessionPrologue()
 	if (!bOk)
 	{
 		UE_TRACE_ERRORMESSAGE(WriteError, GetLastErrorCode());
-		GWriterState.Close(GDataHandle);
-		GDataHandle = 0;
+		Writer_CloseNoFlush();
 		return false;
 	}
 
@@ -1093,13 +1095,7 @@ static void Writer_InternalShutdown()
 	}
 
 	Writer_WorkerJoin();
-
-	if (GDataHandle)
-	{
-		Writer_FlushSendBuffer();
-		GWriterState.Close(GDataHandle);
-		GDataHandle = 0;
-	}
+	Writer_Close();
 
 	Writer_ShutdownControl();
 	Writer_ShutdownPool();
