Added API to UE::Trace to expose IsTailing, which allows game code to know if the tail buffer is in use.

#[review] [REDACTED_HANDLE]

[RN] Added API call UE::Trace::IsTailing which returns true if the Trace Tail Buffer is active
#tests Built and Tested on Dedicated Server Cooked Build
#Synced-CL 51656536

diff --git a/Engine/Source/Runtime/TraceLog/Private/Trace/Trace.cpp b/Engine/Source/Runtime/TraceLog/Private/Trace/Trace.cpp
--- a/Engine/Source/Runtime/TraceLog/Private/Trace/Trace.cpp
+++ b/Engine/Source/Runtime/TraceLog/Private/Trace/Trace.cpp
@@ -31,6 +31,7 @@ bool	Writer_WriteSnapshotTo(const ANSICHAR*, uint32);
 bool	Writer_SendSnapshotTo(const ANSICHAR*, uint32, uint32);
 bool	Writer_IsTracing();
 bool	Writer_IsTracingTo(uint32 (&OutSessionGuid)[4], uint32 (&OutTraceGuid)[4]);
+bool	Writer_IsTailing();
 bool	Writer_Stop();
 uint32	Writer_GetThreadId();
 void	Writer_Panic();
@@ -190,6 +191,12 @@ bool IsTracingTo(uint32 (&OutSessionGuid)[4], uint32 (&OutTraceGuid)[4])
 	return Private::Writer_IsTracingTo(OutSessionGuid, OutTraceGuid);
 }
 
+////////////////////////////////////////////////////////////////////////////////
+bool IsTailing()
+{
+	return Private::Writer_IsTailing();
+}
+
 ////////////////////////////////////////////////////////////////////////////////
 bool Stop()
 {
diff --git a/Engine/Source/Runtime/TraceLog/Public/Trace/Trace.h b/Engine/Source/Runtime/TraceLog/Public/Trace/Trace.h
--- a/Engine/Source/Runtime/TraceLog/Public/Trace/Trace.h
+++ b/Engine/Source/Runtime/TraceLog/Public/Trace/Trace.h
@@ -378,6 +378,13 @@ UE_TRACE_API bool	IsTracing() UE_TRACE_IMPL(false);
  */
 UE_TRACE_API bool	IsTracingTo(uint32 (&OutSessionGuid)[4], uint32 (&OutTraceGuid)[4]) UE_TRACE_IMPL(false);
 
+/**
+* Checks if TraceLog is currently using the Tail Buffer, which means that trace events are being cached in memory and
+* could be dumped into a file at any time.
+* @return True if the TailBuffer is active, false if not
+*/
+UE_TRACE_API bool	IsTailing() UE_TRACE_IMPL(false);
+
 /**
  * Stops current output if any is active.This will fail there is no active output or if there is already a queued
  * output to be started.
