[Insights] Bump max allowed event description to 1024

#rnx
#Synced-CL 52172597

diff --git a/Engine/Source/Runtime/TraceLog/Private/Trace/Writer.cpp b/Engine/Source/Runtime/TraceLog/Private/Trace/Writer.cpp
--- a/Engine/Source/Runtime/TraceLog/Private/Trace/Writer.cpp
+++ b/Engine/Source/Runtime/TraceLog/Private/Trace/Writer.cpp
@@ -587,14 +587,17 @@ static void Writer_DescribeEvents(FEventNode::FIter Iter)
 	// cannot emit additional data on this thread. So use the un-instrumented
 	// send data functions here.
 	TWriteBufferRedirect<4096> TraceData;
+	// Defines the maximum size an event description is allowed to take. This
+	// controls the minimal available headroom in the above buffer.
+	constexpr uint32 MaxEventDescBytes = 1024;
 
 	while (const FEventNode* Event = Iter.GetNext())
 	{
 		Event->Describe();
 
-		// Flush just in case an NewEvent event will be larger than 512 bytes.
+		// Flush just in case an NewEvent event will be larger than MaxEventDescBytes.
 		uint32 RedirectSize = TraceData.GetSize();
-		if (RedirectSize >= (TraceData.GetCapacity() - 512))
+		if (RedirectSize >= (TraceData.GetCapacity() - MaxEventDescBytes))
 		{
 			Writer_SendDataNoInstr(ETransportTid::Events, TraceData.GetData(), RedirectSize);
 			if (TraceData.GetSize() != RedirectSize)
