[Insights] Send OnConnect events on it's own thread

* Refactored redirect types to support growing redirects.
* Added event flag "OnConnect" to mark events that should be emitted ahead of the tail if emitted during OnConnect callback.
* Specialized scope to write explicit event during on connect to special buffers queue if one is active. If none is active it works as a regular log scope. When active treats sync events as if they were no-sync, i.e. does not increment the serial counter and writes the serial value as `0` in the event header.

#Synced-CL 52616121

diff --git a/Engine/Source/Runtime/TraceLog/Private/Trace/Field.cpp b/Engine/Source/Runtime/TraceLog/Private/Trace/Field.cpp
--- a/Engine/Source/Runtime/TraceLog/Private/Trace/Field.cpp
+++ b/Engine/Source/Runtime/TraceLog/Private/Trace/Field.cpp
@@ -44,7 +44,7 @@ static void Field_WriteAuxData(uint32 Index, int32 Size, CallbackType&& Callback
 		}
 		bCommit = true;
 
-		Buffer = Writer_NextBuffer();
+		Buffer = Writer_NextBuffer(Buffer);
 		Remaining = int32(ptrdiff_t((uint8*)Buffer - Buffer->ReadCursor()));
 	};
 
diff --git a/Engine/Source/Runtime/TraceLog/Private/Trace/TlsBuffer.cpp b/Engine/Source/Runtime/TraceLog/Private/Trace/TlsBuffer.cpp
--- a/Engine/Source/Runtime/TraceLog/Private/Trace/TlsBuffer.cpp
+++ b/Engine/Source/Runtime/TraceLog/Private/Trace/TlsBuffer.cpp
@@ -9,6 +9,7 @@
 #include "Platform.h"
 #include "Trace/Detail/Atomic.h"
 #include "Trace/Detail/Writer.inl"
+#include "Trace/Detail/Transport.h"
 #include "Trace/Trace.inl"
 #include "WriteBufferRedirect.h"
 
@@ -23,18 +24,16 @@ uint32				Writer_GetThreadId();
 void				Writer_FreeBlockListToPool(FWriteBuffer*, FWriteBuffer*);
 void				Writer_MemoryFree(void* Address, uint32 Size);
 void*				Writer_MemoryAllocate(SIZE_T Size, uint32 Alignment);
+void				Writer_SendData(uint32, uint8 *__restrict, uint32);
+FWriteBuffer*		Writer_NextBuffer(FWriteBuffer* CurrentBuffer);
 extern uint64		GStartCycle;
 extern FStatistics	GTraceStatistics;
 
-
-
 ////////////////////////////////////////////////////////////////////////////////
 UE_TRACE_MINIMAL_EVENT_BEGIN($Trace, ThreadTiming, NoSync)
 	UE_TRACE_MINIMAL_EVENT_FIELD(uint64, BaseTimestamp)
 UE_TRACE_MINIMAL_EVENT_END()
 
-
-
 ////////////////////////////////////////////////////////////////////////////////
 struct FThreadBuffer
 {
@@ -43,10 +42,11 @@ struct FThreadBuffer
 	uint32 Padding;
 };
 
+
 ////////////////////////////////////////////////////////////////////////////////
 #define T_ALIGN alignas(PLATFORM_CACHE_LINE_SIZE)
 static FWriteBuffer						GNullWriteBuffer	= { {}, static_cast<uint16>(~0u), 0, 0, 0, 0, nullptr, nullptr, nullptr, nullptr, (uint8*)&GNullWriteBuffer };
-thread_local FWriteBuffer*				GTlsWriteBuffer		= &GNullWriteBuffer;
+thread_local FTlsState					GTlsState = {&GNullWriteBuffer, nullptr};
 static FWriteBuffer* __restrict			GActiveThreadList;	// = nullptr;
 T_ALIGN static FWriteBuffer* volatile	GNewThreadList;		// = nullptr;
 static FThreadBuffer*					GThreadBufferArray;	// = nullptr;
@@ -59,10 +59,30 @@ TRACELOG_API FWriteBuffer* Writer_GetBuffer()
 {
 	// Thread locals and DLLs don't mix so for modular builds we are forced to
 	// export this function to access thread-local variables.
-	return GTlsWriteBuffer;
+	return GTlsState.WriteBuffer;
 }
 #endif
 
+////////////////////////////////////////////////////////////////////////////////
+TRACELOG_API bool Writer_OnConnectActivate(FTlsState& OutPrevState)
+{
+	FWriteBuffer* OnConnectBuffer = FOnConnectRedirectScope::GetBuffer();
+	if (OnConnectBuffer == nullptr)
+	{
+		return false;
+	}
+	OutPrevState = GTlsState;
+	GTlsState.WriteBuffer = OnConnectBuffer;
+	GTlsState.NextBufferFn = &FOnConnectRedirectScope::NextBuffer;
+	return true;
+}
+
+////////////////////////////////////////////////////////////////////////////////
+TRACELOG_API void Writer_OnConnectDeactivate(const FTlsState& PrevState)
+{
+	GTlsState = PrevState;
+}
+
 ////////////////////////////////////////////////////////////////////////////////
 static FWriteBuffer* Writer_NextBufferInternal(FWriteBuffer* CurrentBuffer)
 {
@@ -70,7 +90,7 @@ static FWriteBuffer* Writer_NextBufferInternal(FWriteBuffer* CurrentBuffer)
 	// memory. Such events should be tightly controlled and kept to a minimum.
 	// There is not a lot of buffer space available in TraceData to avoid
 	// exhausting available stack and leave usable space in the new buffer.
-	TWriteBufferRedirect<2048> TraceData;
+	TBoundedRedirectScope<2048> TraceData;
 
 #if !UE_BUILD_SHIPPING && !UE_BUILD_TEST
 	if (CurrentBuffer->ThreadId == decltype(TraceData)::ActiveRedirection)
@@ -99,7 +119,7 @@ static FWriteBuffer* Writer_NextBufferInternal(FWriteBuffer* CurrentBuffer)
 
 	TraceData.Abandon();
 
-	GTlsWriteBuffer = NextBuffer;
+	GTlsState.WriteBuffer = NextBuffer;
 
 	NextBuffer->Committed = NextBuffer->ReadCursor();
 
@@ -143,9 +163,14 @@ static FWriteBuffer* Writer_NextBufferInternal(FWriteBuffer* CurrentBuffer)
 }
 
 ////////////////////////////////////////////////////////////////////////////////
-TRACELOG_API FWriteBuffer* Writer_NextBuffer()
+TRACELOG_API FWriteBuffer* Writer_NextBuffer(FWriteBuffer* CurrentBuffer)
 {
-	FWriteBuffer* CurrentBuffer = GTlsWriteBuffer;
+	// Allow us to redirect next buffer function (e.g. growable redirects). If
+	// not set fall back to the default implementation.
+	if (GTlsState.NextBufferFn)
+	{
+		return GTlsState.NextBufferFn(CurrentBuffer);
+	}
 	return Writer_NextBufferInternal(CurrentBuffer);
 }
 
@@ -311,12 +336,12 @@ void Writer_DrainBuffers()
 ////////////////////////////////////////////////////////////////////////////////
 void Writer_DrainLocalBuffers()
 {
-	if (GTlsWriteBuffer == &GNullWriteBuffer)
+	if (GTlsState.WriteBuffer == &GNullWriteBuffer)
 	{
 		return;
 	}
 
-	const uint32 LocalThreadId = GTlsWriteBuffer->ThreadId;
+	const uint32 LocalThreadId = GTlsState.WriteBuffer->ThreadId;
 
 	struct FRetireList
 	{
@@ -380,17 +405,37 @@ void Writer_DrainLocalBuffers()
 	}
 }
 
+////////////////////////////////////////////////////////////////////////////////
+void Writer_DrainBufferList(uint32 ThreadId, FWriteBuffer* List)
+{
+	FWriteBuffer* Buffer = List;
+	do {
+		uint8 *Committed = AtomicLoadAcquire((uint8 **)&Buffer->Committed);
+
+		// Send as much as we can.
+		if (uint32 SizeToReap = uint32(Committed - Buffer->Reaped))
+		{
+#if TRACE_PRIVATE_STATISTICS
+			AtomicAddRelaxed(&GTraceStatistics.BytesTraced, uint64(SizeToReap));
+#endif
+
+			Writer_SendData(ThreadId, Buffer->Reaped, SizeToReap);
+			Buffer->Reaped = Committed;
+		}
+	} while ((Buffer = Buffer->NextBuffer));
+}
+
 ////////////////////////////////////////////////////////////////////////////////
 void Writer_EndThreadBuffer()
 {
-	if (GTlsWriteBuffer == &GNullWriteBuffer)
+	if (GTlsState.WriteBuffer == &GNullWriteBuffer)
 	{
 		return;
 	}
 
-	const PTRINT CurrentCursor = (PTRINT)GTlsWriteBuffer->ReadCursor();
-	int32 EtxOffset = int32(PTRINT((uint8*)GTlsWriteBuffer - CurrentCursor));
-	AtomicStoreRelease(&(GTlsWriteBuffer->EtxOffset), EtxOffset);
+	const PTRINT CurrentCursor = (PTRINT)GTlsState.WriteBuffer->ReadCursor();
+	int32 EtxOffset = int32(PTRINT((uint8*)GTlsState.WriteBuffer - CurrentCursor));
+	AtomicStoreRelease(&(GTlsState.WriteBuffer->EtxOffset), EtxOffset);
 }
 
 } // namespace Private
diff --git a/Engine/Source/Runtime/TraceLog/Private/Trace/WriteBufferRedirect.cpp b/Engine/Source/Runtime/TraceLog/Private/Trace/WriteBufferRedirect.cpp
new file mode 100644
--- /dev/null
+++ b/Engine/Source/Runtime/TraceLog/Private/Trace/WriteBufferRedirect.cpp
@@ -0,0 +1,27 @@
+// Copyright Epic Games, Inc. All Rights Reserved.
+#include "WriteBufferRedirect.h"
+#include "Trace/Config.h"
+
+namespace UE {
+namespace Trace {
+namespace Private {
+
+#if TRACE_PRIVATE_MINIMAL_ENABLED
+
+////////////////////////////////////////////////////////////////////////////////
+FWriteBuffer* FOnConnectRedirectScope::Tail = nullptr;
+
+////////////////////////////////////////////////////////////////////////////////
+FWriteBuffer* FOnConnectRedirectScope::NextBuffer(FWriteBuffer* CurrentBuffer)
+{
+	Tail = TWriteBufferQueue<OnConnectBufferSize>::NextBuffer(CurrentBuffer);
+	return Tail;
+}
+
+////////////////////////////////////////////////////////////////////////////////
+
+#endif //TRACE_PRIVATE_MINIMAL_ENABLED
+
+} // namespace Private
+} // namespace Trace
+} // namespace UE
diff --git a/Engine/Source/Runtime/TraceLog/Private/Trace/WriteBufferRedirect.h b/Engine/Source/Runtime/TraceLog/Private/Trace/WriteBufferRedirect.h
--- a/Engine/Source/Runtime/TraceLog/Private/Trace/WriteBufferRedirect.h
+++ b/Engine/Source/Runtime/TraceLog/Private/Trace/WriteBufferRedirect.h
@@ -2,26 +2,196 @@
 
 #pragma once
 
-// HEADER_UNIT_SKIP - Internal
+#include "Trace/Config.h"
+
+#if TRACE_PRIVATE_MINIMAL_ENABLED
 
+// HEADER_UNIT_SKIP - Internal
+#include "Trace/Trace.h"
 #include "Trace/Detail/Writer.inl"
+#include "Trace/Detail/Transport.h"
 
 namespace UE {
 namespace Trace {
 namespace Private {
 
 ////////////////////////////////////////////////////////////////////////////////
-extern thread_local FWriteBuffer* GTlsWriteBuffer;
+void* Writer_MemoryAllocate(SIZE_T, uint32);
+void  Writer_MemoryFree(void*, uint32);
+
+extern thread_local FTlsState GTlsState;
+struct FWriteBuffer;
+
+////////////////////////////////////////////////////////////////////////////////
+
+/**
+ * A buffer queue of traced data. Usable for redirecting behaviour.
+ */
+template <uint16 BufferSize>
+class TWriteBufferQueue 
+{
+public:
+	TWriteBufferQueue();
+	~TWriteBufferQueue();
+	static FWriteBuffer* AllocBuffer();
+	static FWriteBuffer* FreeBuffer(FWriteBuffer*);
+	static FWriteBuffer* NextBuffer(FWriteBuffer*);
+
+	FWriteBuffer*	Head;
+};
+
+////////////////////////////////////////////////////////////////////////////////
+template <uint16 BufferSize>
+TWriteBufferQueue<BufferSize>::TWriteBufferQueue()
+{
+	Head = AllocBuffer();
+}
+
+////////////////////////////////////////////////////////////////////////////////
+template <uint16 BufferSize>
+TWriteBufferQueue<BufferSize>::~TWriteBufferQueue()
+{
+	FWriteBuffer* NextBuffer = Head;
+	while (NextBuffer)
+	{
+		NextBuffer = FreeBuffer(NextBuffer);
+	}
+}
+
+////////////////////////////////////////////////////////////////////////////////
+template <uint16 BufferSize>
+FWriteBuffer* TWriteBufferQueue<BufferSize>::AllocBuffer()
+{
+	const uint16 UsableBufferSize = BufferSize - sizeof(FWriteBuffer);
+
+	uint8* AllocBase = (uint8*) Writer_MemoryAllocate(BufferSize, PLATFORM_CACHE_LINE_SIZE);
+
+	FWriteBuffer* NextBuffer = (FWriteBuffer*)(AllocBase + UsableBufferSize);
+	NextBuffer->Size = UsableBufferSize;
+	NextBuffer->CursorMaskInv = 0;
+
+	// Reserve some space for packet header
+	uint8* DataStart = AllocBase + sizeof(FTidPacketBase);
+	NextBuffer->SetCursor(DataStart);
+	NextBuffer->Committed = DataStart;
+	NextBuffer->Reaped = DataStart;
+
+	NextBuffer->EtxOffset = 0 - int32(sizeof(FWriteBuffer));
+	NextBuffer->ThreadId = 0;
+	NextBuffer->NextBuffer = nullptr;
+
+	return NextBuffer;
+}
+
+////////////////////////////////////////////////////////////////////////////////
+template <uint16 BufferSize>
+FWriteBuffer* TWriteBufferQueue<BufferSize>::FreeBuffer(FWriteBuffer* Buffer)
+{
+	FWriteBuffer* NextBuffer = Buffer->NextBuffer;
+	const uint32 AllocSize = sizeof(FWriteBuffer) + Buffer->Size;
+	uint8* AllocBase = ((uint8*)Buffer) - Buffer->Size;
+
+	Writer_MemoryFree(AllocBase, AllocSize);
+
+	return NextBuffer;
+}
+
+////////////////////////////////////////////////////////////////////////////////
+template <uint16 BufferSize>
+FWriteBuffer* TWriteBufferQueue<BufferSize>::NextBuffer(FWriteBuffer* __restrict PrevBuffer)
+{
+	FWriteBuffer* NextBuffer = AllocBuffer();
+	PrevBuffer->NextBuffer = NextBuffer;
+	GTlsState.WriteBuffer = NextBuffer;
+	return NextBuffer;
+}
+
+
+////////////////////////////////////////////////////////////////////////////////
+
+/**
+ * Unbounded redirect buffer. Will allocate new buffers with templated size 
+ * while active.
+ */
+template <uint16 BufferSize>
+class TRedirectScope : public TWriteBufferQueue<BufferSize>
+{
+public:
+	enum : uint16 { ActiveGrowingRedirection = 0xfffe };
+
+	TRedirectScope();
+	~TRedirectScope();
+
+	FWriteBuffer* Finish();
+	static FWriteBuffer* NextBuffer(FWriteBuffer* CurrentBuffer);
+
+private:
+	void ResetThreadLocalState();
+
+	FTlsState		Prev;
+};
+
+////////////////////////////////////////////////////////////////////////////////
+template <uint16 BufferSize>
+TRedirectScope<BufferSize>::TRedirectScope()
+{
+	this->Head->ThreadId = ActiveGrowingRedirection;
+	Prev = GTlsState;
+	GTlsState.WriteBuffer = this->Head;
+	GTlsState.NextBufferFn = &TRedirectScope::NextBuffer;
+}
+
+////////////////////////////////////////////////////////////////////////////////
+template <uint16 BufferSize>
+TRedirectScope<BufferSize>::~TRedirectScope()
+{
+	ResetThreadLocalState();
+}
+
+////////////////////////////////////////////////////////////////////////////////
+template <uint16 BufferSize>
+void TRedirectScope<BufferSize>::ResetThreadLocalState()
+{
+	if (Prev.WriteBuffer != nullptr)
+	{
+		GTlsState = Prev;
+		Prev = {nullptr, nullptr};
+	}
+}
+
+////////////////////////////////////////////////////////////////////////////////
+template <uint16 BufferSize>
+FWriteBuffer* TRedirectScope<BufferSize>::NextBuffer(FWriteBuffer* CurrentBuffer)
+{
+	FWriteBuffer* NextBuffer = TWriteBufferQueue<BufferSize>::NextBuffer(CurrentBuffer);
+	NextBuffer->ThreadId = ActiveGrowingRedirection;
+	return NextBuffer;
+}
+
+////////////////////////////////////////////////////////////////////////////////
+template <uint16 BufferSize>
+FWriteBuffer* TRedirectScope<BufferSize>::Finish()
+{
+	FWriteBuffer* FirstBuffer = this->Head;
+	this->Head = nullptr;
+	ResetThreadLocalState();
+	return FirstBuffer;
+}
+
 
 ////////////////////////////////////////////////////////////////////////////////
-template <int BufferSize>
-class TWriteBufferRedirect
+
+/**
+ * Bounded redirect buffer.
+ */
+template <uint16 BufferSize>
+class TBoundedRedirectScope
 {
 public:
 	enum : uint16 { ActiveRedirection = 0xffff };
 
-					TWriteBufferRedirect();
-					~TWriteBufferRedirect();
+					TBoundedRedirectScope();
+					~TBoundedRedirectScope();
 	void			Close();
 	void			Abandon();
 	uint8*			GetData();
@@ -30,73 +200,74 @@ public:
 	void			Reset();
 
 private:
-	FWriteBuffer*	PrevBuffer;
+	FTlsState		Prev;
 	uint8			Data[BufferSize];
 	FWriteBuffer	Buffer;
 };
 
 ////////////////////////////////////////////////////////////////////////////////
-template <int BufferSize>
-inline TWriteBufferRedirect<BufferSize>::TWriteBufferRedirect()
+template <uint16 BufferSize>
+inline TBoundedRedirectScope<BufferSize>::TBoundedRedirectScope()
 {
 	Reset();
-	PrevBuffer = GTlsWriteBuffer;
-	GTlsWriteBuffer = &Buffer;
+	Prev = GTlsState;
+	GTlsState.WriteBuffer = &Buffer;
+	GTlsState.NextBufferFn = nullptr;
 	Buffer.Size = uint16(BufferSize);
+	// When using the default NextBuffer function it looks out
+	// for this thread id to avoid overflowing
 	Buffer.ThreadId = ActiveRedirection;
 }
 
 ////////////////////////////////////////////////////////////////////////////////
-template <int BufferSize>
-inline TWriteBufferRedirect<BufferSize>::~TWriteBufferRedirect()
+template <uint16 BufferSize>
+inline TBoundedRedirectScope<BufferSize>::~TBoundedRedirectScope()
 {
 	Close();
 }
 
 ////////////////////////////////////////////////////////////////////////////////
-template <int BufferSize>
-inline void TWriteBufferRedirect<BufferSize>::Close()
+template <uint16 BufferSize>
+inline void TBoundedRedirectScope<BufferSize>::Close()
 {
-	if (PrevBuffer == nullptr)
+	if (Prev.WriteBuffer != nullptr)
 	{
-		return;
+		GTlsState = Prev;
+		Prev = {nullptr, nullptr};
 	}
-
-	GTlsWriteBuffer = PrevBuffer;
-	PrevBuffer = nullptr;
 }
 
 ////////////////////////////////////////////////////////////////////////////////
-template <int BufferSize>
-inline void TWriteBufferRedirect<BufferSize>::Abandon()
+template <uint16 BufferSize>
+inline void TBoundedRedirectScope<BufferSize>::Abandon()
 {
-	PrevBuffer = nullptr;
+	Prev = {nullptr, nullptr};
 }
 
 ////////////////////////////////////////////////////////////////////////////////
-template <int BufferSize>
-inline uint8* TWriteBufferRedirect<BufferSize>::GetData()
+template <uint16 BufferSize>
+inline uint8* TBoundedRedirectScope<BufferSize>::GetData()
 {
 	return Buffer.Reaped;
 }
 
 ////////////////////////////////////////////////////////////////////////////////
-template <int BufferSize>
-inline uint32 TWriteBufferRedirect<BufferSize>::GetSize() const
+template <uint16 BufferSize>
+inline uint32 TBoundedRedirectScope<BufferSize>::GetSize() const
 {
 	return uint32(Buffer.Committed - Buffer.Reaped);
 }
 
 ////////////////////////////////////////////////////////////////////////////////
-template <int BufferSize>
-inline uint32 TWriteBufferRedirect<BufferSize>::GetCapacity() const
+template <uint16 BufferSize>
+inline uint32 TBoundedRedirectScope<BufferSize>::GetCapacity() const
 {
 	return BufferSize;
 }
 
 ////////////////////////////////////////////////////////////////////////////////
-template <int BufferSize>
-inline void TWriteBufferRedirect<BufferSize>::Reset()
+template <uint16 BufferSize>
+inline void TBoundedRedirectScope<BufferSize>::Reset()
 {
 	Buffer.CursorMaskInv = 0;
 	Buffer.SetCursor(Data + sizeof(uint32));
@@ -104,6 +275,47 @@ inline void TWriteBufferRedirect<BufferSize>::Reset()
 	Buffer.Reaped = Buffer.ReadCursor();
 }
 
+
+////////////////////////////////////////////////////////////////////////////////
+
+constexpr uint16 OnConnectBufferSize = 512;
+
+/**
+ * Specialized redirect buffer queue for OnConnect callback.
+ */
+class FOnConnectRedirectScope : public TWriteBufferQueue<OnConnectBufferSize>
+{
+public:
+	FOnConnectRedirectScope();
+	~FOnConnectRedirectScope();
+	static FWriteBuffer* GetBuffer();
+	static FWriteBuffer* NextBuffer(FWriteBuffer* CurrentBuffer);
+
+private:
+	static FWriteBuffer* Tail;
+};
+
+////////////////////////////////////////////////////////////////////////////////
+inline FOnConnectRedirectScope::FOnConnectRedirectScope()
+{
+	Tail = Head;
+}
+
+////////////////////////////////////////////////////////////////////////////////
+inline FOnConnectRedirectScope::~FOnConnectRedirectScope()
+{
+	Tail = nullptr;
+}
+
+////////////////////////////////////////////////////////////////////////////////
+inline FWriteBuffer* FOnConnectRedirectScope::GetBuffer()
+{
+	return Tail;
+}
+
+
 } // namespace Private
 } // namespace Trace
 } // namespace UE
+  
+#endif // TRACE_PRIVATE_MINIMAL_ENABLED
diff --git a/Engine/Source/Runtime/TraceLog/Private/Trace/Writer.cpp b/Engine/Source/Runtime/TraceLog/Private/Trace/Writer.cpp
--- a/Engine/Source/Runtime/TraceLog/Private/Trace/Writer.cpp
+++ b/Engine/Source/Runtime/TraceLog/Private/Trace/Writer.cpp
@@ -56,6 +56,7 @@ void			Writer_InitializePool();
 void			Writer_ShutdownPool();
 void			Writer_DrainBuffers();
 void			Writer_DrainLocalBuffers();
+void			Writer_DrainBufferList(uint32, FWriteBuffer*);
 void			Writer_EndThreadBuffer();
 uint32			Writer_GetControlPort();
 void			Writer_UpdateControl();
@@ -127,6 +128,7 @@ static FWriterState GWriterState;
 static FWriterState GPendingWriterState;
 static uint8 GPendingWriterLock;
 
+
 ////////////////////////////////////////////////////////////////////////////////
 // When a thread terminates we want to recover its trace buffer. To do this we
 // use a thread_local object whose destructor gets called as the thread ends. On
@@ -586,7 +588,7 @@ static void Writer_DescribeEvents(FEventNode::FIter Iter)
 	// Since we are using on stack buffer to collect the describing events we
 	// cannot emit additional data on this thread. So use the un-instrumented
 	// send data functions here.
-	TWriteBufferRedirect<4096> TraceData;
+	TBoundedRedirectScope<4096> TraceData;
 	// Defines the maximum size an event description is allowed to take. This
 	// controls the minimal available headroom in the above buffer.
 	constexpr uint32 MaxEventDescBytes = 1024;
@@ -936,19 +938,17 @@ void Writer_CallbackOnConnect()
 		return;
 	}
 
-	// Prior to letting callbacks trace events we need to flush any pending
-	// trace data to that tail. We do not want that data to be sent over the
-	// wire as that would cause data to be sent out-of-order.
-	UPTRINT DataHandle = GDataHandle;
-	AtomicStoreRelaxed(&GDataHandle, (UPTRINT)0);
-	Writer_DrainLocalBuffers();
-	AtomicStoreRelaxed(&GDataHandle, DataHandle);
+	// In order for events to appear ahead of the tail emit
+	// them on a special thread id.
+	FOnConnectRedirectScope OnConnectEvents;
 
 	// Issue callback. We assume any events emitted here are not marked as
 	// important and emitted on this thread.
 	GOnConnection();
-}
 
+	// Write the contents of the redirect buffers to the stream
+	Writer_DrainBufferList(ETransportTid::PseudoImportants, OnConnectEvents.Head);
+}
 
 ////////////////////////////////////////////////////////////////////////////////
 static UPTRINT			GWorkerThread;		// = 0;
diff --git a/Engine/Source/Runtime/TraceLog/Public/Trace/Detail/EventNode.h b/Engine/Source/Runtime/TraceLog/Public/Trace/Detail/EventNode.h
--- a/Engine/Source/Runtime/TraceLog/Public/Trace/Detail/EventNode.h
+++ b/Engine/Source/Runtime/TraceLog/Public/Trace/Detail/EventNode.h
@@ -26,6 +26,7 @@ struct FEventInfo
 		Flag_Definition16	= 1 << 4,
 		Flag_Definition32	= 1 << 5,
 		Flag_Definition64	= 1 << 6,
+		Flag_MaybeImportant	= 1 << 7,
 
 		DefinitionBits		= Flag_Definition8 | Flag_Definition16 | Flag_Definition32 | Flag_Definition64,
 	};
diff --git a/Engine/Source/Runtime/TraceLog/Public/Trace/Detail/LogScope.h b/Engine/Source/Runtime/TraceLog/Public/Trace/Detail/LogScope.h
--- a/Engine/Source/Runtime/TraceLog/Public/Trace/Detail/LogScope.h
+++ b/Engine/Source/Runtime/TraceLog/Public/Trace/Detail/LogScope.h
@@ -7,6 +7,7 @@
 #if TRACE_PRIVATE_MINIMAL_ENABLED
 
 #include "Writer.inl" // IWYU pragma: export
+#include <utility>
 
 namespace UE {
 namespace Trace {
@@ -41,12 +42,13 @@ private:
 	template <uint32 Flags>
 	static auto				EnterImpl(uint32 Uid, uint32 Size);
 	template <class T> inline void	EnterPrelude(uint32 Size);
-	inline void				Enter(uint32 Uid, uint32 Size);
+	inline void				Enter(uint32 Uid, uint32 Size, bool bSetSerial = true);
 	inline void				EnterNoSync(uint32 Uid, uint32 Size);
 	uint8*					Ptr;
 	FWriteBuffer*			Buffer;
 };
 
+////////////////////////////////////////////////////////////////////////////////
 template <bool bMaybeHasAux>
 class TLogScope
 	: public FLogScope
@@ -55,6 +57,24 @@ public:
 	inline void				operator += (const FLogScope&) const;
 };
 
+////////////////////////////////////////////////////////////////////////////////
+template <bool bMaybeHasAux>
+class TOnConnectLogScope : public TLogScope<bMaybeHasAux>
+{
+	friend class FLogScope;
+public:
+	TOnConnectLogScope();
+	TOnConnectLogScope(TOnConnectLogScope&& Other)
+		: TLogScope<bMaybeHasAux>(std::move(Other))
+		, PrevState(Other.PrevState)
+		, bActive(Other.bActive)
+	{ Other.bActive = false; };
+	~TOnConnectLogScope();
+
+private:
+	FTlsState PrevState;
+	bool bActive;
+};
 
 
 ////////////////////////////////////////////////////////////////////////////////
diff --git a/Engine/Source/Runtime/TraceLog/Public/Trace/Detail/LogScope.inl b/Engine/Source/Runtime/TraceLog/Public/Trace/Detail/LogScope.inl
--- a/Engine/Source/Runtime/TraceLog/Public/Trace/Detail/LogScope.inl
+++ b/Engine/Source/Runtime/TraceLog/Public/Trace/Detail/LogScope.inl
@@ -46,16 +46,41 @@ template <uint32 Flags>
 UE_AUTORTFM_ALWAYS_OPEN
 inline auto FLogScope::EnterImpl(uint32 Uid, uint32 Size)
 {
-	TLogScope<(Flags & FEventInfo::Flag_MaybeHasAux) != 0> Ret;
-	if ((Flags & FEventInfo::Flag_NoSync) != 0)
+	constexpr bool bMaybeHasAux = (Flags & FEventInfo::Flag_MaybeHasAux) != 0;
+
+	if constexpr ((Flags & FEventInfo::Flag_MaybeImportant) != 0)
 	{
-		Ret.EnterNoSync(Uid, Size);
+		// If a on connect buffer queue is active this scope
+		// will redirect the event to this separate queue. If not
+		// it will act exactly like TLogScope.
+		TOnConnectLogScope<bMaybeHasAux> Ret;
+		if ((Flags & FEventInfo::Flag_NoSync) != 0)
+		{
+			Ret.EnterNoSync(Uid, Size);
+		}
+		else
+		{
+			// When the TOnConnectLogScope is active we're in an OnConnect
+			// callback and this event moved ahead of the tail which means 
+			// we don't set the serial number for synchronized events. 
+			const bool bSetSerial = Ret.bActive ? false : true;
+			Ret.Enter(Uid, Size, bSetSerial);
+		}
+		return Ret;
 	}
 	else
 	{
-		Ret.Enter(Uid, Size);
+		TLogScope<bMaybeHasAux> Ret;
+		if ((Flags & FEventInfo::Flag_NoSync) != 0)
+		{
+			Ret.EnterNoSync(Uid, Size);
+		}
+		else
+		{
+			Ret.Enter(Uid, Size);
+		}
+		return Ret;
 	}
-	return Ret;
 }
 
 ////////////////////////////////////////////////////////////////////////////////
@@ -68,7 +93,7 @@ inline void FLogScope::EnterPrelude(uint32 Size)
 	Buffer = Writer_GetBuffer();
 	if (UNLIKELY(Buffer->ReadCursor() + AllocSize > (uint8*)Buffer))
 	{
-		Buffer = Writer_NextBuffer();
+		Buffer = Writer_NextBuffer(Buffer);
 	}
 
 #if !UE_BUILD_SHIPPING && !UE_BUILD_TEST
@@ -87,12 +112,12 @@ inline void FLogScope::EnterPrelude(uint32 Size)
 
 ////////////////////////////////////////////////////////////////////////////////
 UE_AUTORTFM_ALWAYS_OPEN
-inline void FLogScope::Enter(uint32 Uid, uint32 Size)
+inline void FLogScope::Enter(uint32 Uid, uint32 Size, bool bSetSerial /* = true */)
 {
 	EnterPrelude<FEventHeaderSync>(Size);
 
 	uint16 Uid16 = uint16(Uid) | int32(EKnownEventUids::Flag_TwoByteUid);
-	uint32 Serial = uint32(AtomicAddRelaxed(&GLogSerial, 1u));
+	uint32 Serial = bSetSerial ? uint32(AtomicAddRelaxed(&GLogSerial, 1u)) : 0u;
 
 	// Event header FEventHeaderSync
 	memcpy(Ptr - 3, &Serial, sizeof(Serial)); /* FEventHeaderSync::SerialHigh,SerialLow */
@@ -154,7 +179,7 @@ inline void FScopedLogScope::Deinit()
 	FWriteBuffer* Buffer = Writer_GetBuffer();
 	if (UNLIKELY(int32((uint8*)Buffer - Buffer->ReadCursor()) < int32(sizeof(LeaveUid))))
 	{
-		Buffer = Writer_NextBuffer();
+		Buffer = Writer_NextBuffer(Buffer);
 	}
 
 	*Buffer->ReadCursor() = LeaveUid;
@@ -214,7 +239,7 @@ inline void FScopedStampedLogScope::Deinit()
 	FWriteBuffer* Buffer = Writer_GetBuffer();
 	if (UNLIKELY(int32((uint8*)Buffer - Buffer->ReadCursor()) < int32(sizeof(Stamp))))
 	{
-		Buffer = Writer_NextBuffer();
+		Buffer = Writer_NextBuffer(Buffer);
 	}
 
 	Stamp <<= 8;
@@ -275,7 +300,7 @@ FORCENOINLINE auto FLogScope::ScopedEnter()
 	constexpr int32 RequiredSize = sizeof(EnterUid) + EventHeaderSize + EventType::GetSize();
 	if (UNLIKELY(int32((uint8*)Buffer - Buffer->ReadCursor()) < RequiredSize))
 	{
-		Buffer = Writer_NextBuffer();
+		Buffer = Writer_NextBuffer(Buffer);
 	}
 
 	*Buffer->ReadCursor() = EnterUid;
@@ -298,7 +323,7 @@ FORCENOINLINE auto FLogScope::ScopedStampedEnter()
 	constexpr int32 RequiredSize = sizeof(Stamp) + EventHeaderSize + EventType::GetSize();
 	if (UNLIKELY(int32((uint8*)Buffer - Buffer->ReadCursor()) < RequiredSize))
 	{
-		Buffer = Writer_NextBuffer();
+		Buffer = Writer_NextBuffer(Buffer);
 	}
 
 	Stamp = TimeGetTimestamp() - GStartCycle;
@@ -407,6 +432,26 @@ struct FLogScope::FFieldSet<FieldMeta, TEventRef<DefinitionType>>
 	}
 };
 
+////////////////////////////////////////////////////////////////////////////////
+template <bool bMaybeHasAux>
+TOnConnectLogScope<bMaybeHasAux>::TOnConnectLogScope()
+{
+	bActive = Writer_OnConnectActivate(PrevState);
+}
+
+////////////////////////////////////////////////////////////////////////////////
+template <bool bMaybeHasAux>
+TOnConnectLogScope<bMaybeHasAux>::~TOnConnectLogScope()
+{
+	if (bActive)
+	{
+		Writer_OnConnectDeactivate(PrevState);
+	}
+}
+
+
+
+
 } // namespace Private
 } // namespace Trace
 } // namespace UE
diff --git a/Engine/Source/Runtime/TraceLog/Public/Trace/Detail/Trace.h b/Engine/Source/Runtime/TraceLog/Public/Trace/Detail/Trace.h
--- a/Engine/Source/Runtime/TraceLog/Public/Trace/Detail/Trace.h
+++ b/Engine/Source/Runtime/TraceLog/Public/Trace/Detail/Trace.h
@@ -93,6 +93,7 @@ class FChannel;
 		{ \
 			Important			= UE::Trace::Private::FEventInfo::Flag_Important, \
 			NoSync				= UE::Trace::Private::FEventInfo::Flag_NoSync, \
+			MaybeImportant		= UE::Trace::Private::FEventInfo::Flag_MaybeImportant, \
 			Definition8bit		= UE::Trace::Private::FEventInfo::Flag_Definition8, \
 			Definition16bit		= UE::Trace::Private::FEventInfo::Flag_Definition16, \
 			Definition32bit		= UE::Trace::Private::FEventInfo::Flag_Definition32, \
diff --git a/Engine/Source/Runtime/TraceLog/Public/Trace/Detail/Transport.h b/Engine/Source/Runtime/TraceLog/Public/Trace/Detail/Transport.h
--- a/Engine/Source/Runtime/TraceLog/Public/Trace/Detail/Transport.h
+++ b/Engine/Source/Runtime/TraceLog/Public/Trace/Detail/Transport.h
@@ -26,7 +26,8 @@ enum ETransportTid : uint32
 	Importants	= Internal,		// important/cached events
 	Bias,						// [Bias,End] = threads. Note bias can't be..
 	/* ... */					// ..changed as it breaks backwards compat :(
-	End			= 0x3ffe,		// two msbs are user for packet markers
+	End			= 0x3ff0,		// two msbs are user for packet markers
+	PseudoImportants	= 0x3ffe,	// see Writer_CallbackOnConnect()
 	Sync		= 0x3fff,		// see Writer_SendSync()
 };
 
diff --git a/Engine/Source/Runtime/TraceLog/Public/Trace/Detail/Writer.inl b/Engine/Source/Runtime/TraceLog/Public/Trace/Detail/Writer.inl
--- a/Engine/Source/Runtime/TraceLog/Public/Trace/Detail/Writer.inl
+++ b/Engine/Source/Runtime/TraceLog/Public/Trace/Detail/Writer.inl
@@ -56,17 +56,27 @@ inline uint8* FWriteBuffer::ReadCursor() const
 	return Cursor;
 }
 
+////////////////////////////////////////////////////////////////////////////////
+using NextBufferFunc = struct FWriteBuffer*(struct FWriteBuffer*);
+struct FTlsState
+{
+	FWriteBuffer*	WriteBuffer;
+	NextBufferFunc*	NextBufferFn;
+};
+
 ////////////////////////////////////////////////////////////////////////////////
 TRACELOG_API uint64				TimeGetTimestamp();
-TRACELOG_API FWriteBuffer*		Writer_NextBuffer();
+TRACELOG_API FWriteBuffer*		Writer_NextBuffer(FWriteBuffer*);
 TRACELOG_API FWriteBuffer*		Writer_GetBuffer();
+TRACELOG_API bool				Writer_OnConnectActivate(FTlsState& OutPrevState);
+TRACELOG_API void				Writer_OnConnectDeactivate(const FTlsState& PrevState);
 
 ////////////////////////////////////////////////////////////////////////////////
 #if IS_MONOLITHIC
-extern thread_local FWriteBuffer* GTlsWriteBuffer;
+extern thread_local FTlsState GTlsState;
 inline FWriteBuffer* Writer_GetBuffer()
 {
-	return GTlsWriteBuffer;
+	return GTlsState.WriteBuffer;
 }
 #endif // IS_MONOLITHIC
 
